PluginReferenceViewer - Added history and navigation controls

#rnx

diff --git a/Engine/Plugins/Experimental/PluginReferenceViewer/Source/PluginReferenceViewer/Private/EdGraph_PluginReferenceViewer.cpp b/Engine/Plugins/Experimental/PluginReferenceViewer/Source/PluginReferenceViewer/Private/EdGraph_PluginReferenceViewer.cpp
--- a/Engine/Plugins/Experimental/PluginReferenceViewer/Source/PluginReferenceViewer/Private/EdGraph_PluginReferenceViewer.cpp
+++ b/Engine/Plugins/Experimental/PluginReferenceViewer/Source/PluginReferenceViewer/Private/EdGraph_PluginReferenceViewer.cpp
@@ -51,13 +51,18 @@ void UEdGraph_PluginReferenceViewer::SetGraphRoot(const TArray<FPluginIdentifier
 	CurrentGraphRootOrigin = GraphRootOrigin;
 }
 
+const TArray<FPluginIdentifier>& UEdGraph_PluginReferenceViewer::GetCurrentGraphRootIdentifiers() const
+{
+	return CurrentGraphRootIdentifiers;
+}
+
 void UEdGraph_PluginReferenceViewer::CachePluginDependencies(const TArray<TSharedRef<IPlugin>>& Plugins)
 {
 	PluginMap.Empty();
 
 	for (const TSharedRef<IPlugin>& Plugin : Plugins)
 	{
-		PluginMap.Add(FPluginIdentifier(Plugin->GetName()), Plugin);
+		PluginMap.Add(FPluginIdentifier::FromString(Plugin->GetName()), Plugin);
 	}
 
 	struct FDependencyNodeBuilder
@@ -74,7 +79,8 @@ void UEdGraph_PluginReferenceViewer::CachePluginDependencies(const TArray<TShare
 		{
 			for (const TSharedRef<IPlugin>& Plugin : Plugins)
 			{
-				Nodes.Add(FPluginIdentifier(Plugin->GetName()), MakeUnique<FPluginDependsNode>(Plugin->GetName()));
+				FPluginIdentifier NewPluginIdentifier = FPluginIdentifier::FromString(Plugin->GetName());
+				Nodes.Add(NewPluginIdentifier, MakeUnique<FPluginDependsNode>(NewPluginIdentifier));
 			}
 
 			for (const TSharedRef<IPlugin>& Plugin : Plugins)
@@ -85,14 +91,15 @@ void UEdGraph_PluginReferenceViewer::CachePluginDependencies(const TArray<TShare
 				for (const FPluginReferenceDescriptor& Item : PluginDescriptor.Plugins)
 				{
 					FPluginDependsNode* ChildNode = nullptr;
-					TUniquePtr<FPluginDependsNode>* FoundNodePtr = Nodes.Find(FPluginIdentifier(Item.Name));
+					TUniquePtr<FPluginDependsNode>* FoundNodePtr = Nodes.Find(FPluginIdentifier::FromString(Item.Name));
 					if (FoundNodePtr != nullptr)
 					{
 						ChildNode = FoundNodePtr->Get();
 					}
 					else
 					{
-						ChildNode = Nodes.Add(FPluginIdentifier(Item.Name), MakeUnique<FPluginDependsNode>(Item.Name)).Get();
+						FPluginIdentifier NewPluginIdentifier = FPluginIdentifier::FromString(Item.Name);
+						ChildNode = Nodes.Add(NewPluginIdentifier, MakeUnique<FPluginDependsNode>(NewPluginIdentifier)).Get();
 					}
 
 					ParentNode->Dependencies.Add(ChildNode);
@@ -109,60 +116,50 @@ void UEdGraph_PluginReferenceViewer::CachePluginDependencies(const TArray<TShare
 
 UEdGraphNode_PluginReference* UEdGraph_PluginReferenceViewer::ConstructNodes(const TArray<FPluginIdentifier>& GraphRootIdentifiers, const FIntPoint& GraphRootOrigin)
 {
-	const bool bIsCompactMode = PluginReferenceViewer.Pin()->IsCompactModeChecked();
-	const int32 SettingsReferencerMaxDepth = PluginReferenceViewer.Pin()->GetSearchReferencerDepthCount();
-	const int32 SettingsDependencyMaxDepth = PluginReferenceViewer.Pin()->GetSearchDependencyDepthCount();
-
-	TMap<FPluginIdentifier, FPluginReferenceNodeInfo> NewReferenceNodeInfos;
-	for (const FPluginIdentifier& RootIdentifier : GraphRootIdentifiers)
-	{
-		FPluginReferenceNodeInfo& RootNodeInfo = NewReferenceNodeInfos.FindOrAdd(RootIdentifier, FPluginReferenceNodeInfo(RootIdentifier, true));
-		RootNodeInfo.Parents.Emplace(NAME_None);
-	}
-
-	RecursivelyPopulateNodeInfos(true, GraphRootIdentifiers, NewReferenceNodeInfos, 0, SettingsReferencerMaxDepth);
-
-	TMap<FPluginIdentifier, FPluginReferenceNodeInfo> NewDependencyNodeInfos;
-	for (const FPluginIdentifier& RootIdentifier : GraphRootIdentifiers)
+	if (!GraphRootIdentifiers.IsEmpty())
 	{
-		FPluginReferenceNodeInfo& DRootNodeInfo = NewDependencyNodeInfos.FindOrAdd(RootIdentifier, FPluginReferenceNodeInfo(RootIdentifier, false));
-		DRootNodeInfo.Parents.Emplace(NAME_None);
-	}
+		const int32 SettingsReferencerMaxDepth = PluginReferenceViewer.Pin()->GetSearchReferencerDepthCount();
+		const int32 SettingsDependencyMaxDepth = PluginReferenceViewer.Pin()->GetSearchDependencyDepthCount();
 
-	RecursivelyPopulateNodeInfos(false , GraphRootIdentifiers, NewDependencyNodeInfos, 0, SettingsDependencyMaxDepth);
+		TMap<FPluginIdentifier, FPluginReferenceNodeInfo> NewReferenceNodeInfos;
+		for (const FPluginIdentifier& RootIdentifier : GraphRootIdentifiers)
+		{
+			FPluginReferenceNodeInfo& RootNodeInfo = NewReferenceNodeInfos.FindOrAdd(RootIdentifier, FPluginReferenceNodeInfo(RootIdentifier, true));
+			RootNodeInfo.Parents.Emplace(NAME_None);
+		}
 
-	// Store the Plugin in the NodeInfos
-	for (TPair<FPluginIdentifier, FPluginReferenceNodeInfo>& InfoPair : NewReferenceNodeInfos)
-	{
-		InfoPair.Value.Plugin = PluginMap.FindChecked(InfoPair.Key);
-	}
+		RecursivelyPopulateNodeInfos(true, GraphRootIdentifiers, NewReferenceNodeInfos, 0, SettingsReferencerMaxDepth);
 
-	for (TPair<FPluginIdentifier, FPluginReferenceNodeInfo>& InfoPair : NewDependencyNodeInfos)
-	{
-		InfoPair.Value.Plugin = PluginMap.FindChecked(InfoPair.Key);
-	}
+		TMap<FPluginIdentifier, FPluginReferenceNodeInfo> NewDependencyNodeInfos;
+		for (const FPluginIdentifier& RootIdentifier : GraphRootIdentifiers)
+		{
+			FPluginReferenceNodeInfo& DRootNodeInfo = NewDependencyNodeInfos.FindOrAdd(RootIdentifier, FPluginReferenceNodeInfo(RootIdentifier, false));
+			DRootNodeInfo.Parents.Emplace(NAME_None);
+		}
 
-	ReferencerNodeInfos = NewReferenceNodeInfos;
-	DependencyNodeInfos = NewDependencyNodeInfos;
+		RecursivelyPopulateNodeInfos(false, GraphRootIdentifiers, NewDependencyNodeInfos, 0, SettingsDependencyMaxDepth);
 
-	RemoveAllNodes();
+		// Store the Plugin in the NodeInfos
+		for (TPair<FPluginIdentifier, FPluginReferenceNodeInfo>& InfoPair : NewReferenceNodeInfos)
+		{
+			InfoPair.Value.Plugin = PluginMap.FindChecked(InfoPair.Key);
+		}
 
-	UEdGraphNode_PluginReference* RootNode = nullptr;
+		for (TPair<FPluginIdentifier, FPluginReferenceNodeInfo>& InfoPair : NewDependencyNodeInfos)
+		{
+			InfoPair.Value.Plugin = PluginMap.FindChecked(InfoPair.Key);
+		}
 
-	if (GraphRootIdentifiers.Num() > 0 && (!ReferencerNodeInfos.IsEmpty() || !DependencyNodeInfos.IsEmpty()))
+		ReferencerNodeInfos = NewReferenceNodeInfos;
+		DependencyNodeInfos = NewDependencyNodeInfos;
+	}
+	else
 	{
-		FPluginIdentifier FirstGraphRootIdentifier = GraphRootIdentifiers[0];
-
-		const FPluginReferenceNodeInfo& NodeInfo = ReferencerNodeInfos[FirstGraphRootIdentifier];
-		RootNode = CreatePluginReferenceNode();
-		RootNode->SetupPluginReferenceNode(CurrentGraphRootOrigin, NodeInfo.Identifier, NodeInfo.Plugin, !bIsCompactMode);
-
-		RecursivelyCreateNodes(true, FirstGraphRootIdentifier, CurrentGraphRootOrigin, FirstGraphRootIdentifier, RootNode, ReferencerNodeInfos, 0, SettingsReferencerMaxDepth, /*bIsRoot*/ true);
-		
-		RecursivelyCreateNodes(false, FirstGraphRootIdentifier, CurrentGraphRootOrigin, FirstGraphRootIdentifier, RootNode, DependencyNodeInfos, 0, SettingsDependencyMaxDepth, /*bIsRoot*/ true);
+		ReferencerNodeInfos.Empty();
+		DependencyNodeInfos.Empty();
 	}
 
-	return RootNode;
+	return RefilterGraph();
 }
 
 UEdGraphNode_PluginReference* UEdGraph_PluginReferenceViewer::RecursivelyCreateNodes(bool bInReferencers, const FPluginIdentifier& InPluginId, const FIntPoint& InNodeLoc, const FPluginIdentifier& InParentId, UEdGraphNode_PluginReference* InParentNode, TMap<FPluginIdentifier, FPluginReferenceNodeInfo>& InNodeInfos, int32 InCurrentDepth, int32 InMaxDepth, bool bIsRoot)
@@ -310,6 +307,36 @@ UEdGraphNode_PluginReference* UEdGraph_PluginReferenceViewer::RebuildGraph()
 	return NewRootNode;
 }
 
+UEdGraphNode_PluginReference* UEdGraph_PluginReferenceViewer::RefilterGraph()
+{
+	RemoveAllNodes();
+
+	UEdGraphNode_PluginReference* RootNode = nullptr;
+
+	if (CurrentGraphRootIdentifiers.Num() > 0 && (!ReferencerNodeInfos.IsEmpty() || !DependencyNodeInfos.IsEmpty()))
+	{
+		const TSharedPtr<const SPluginReferenceViewer> LocalPluginReferenceViewer = PluginReferenceViewer.Pin();
+
+		const bool bIsCompactMode = LocalPluginReferenceViewer->IsCompactModeChecked();
+		const int32 SettingsReferencerMaxDepth = LocalPluginReferenceViewer->GetSearchReferencerDepthCount();
+		const int32 SettingsDependencyMaxDepth = LocalPluginReferenceViewer->GetSearchDependencyDepthCount();
+
+		FPluginIdentifier FirstGraphRootIdentifier = CurrentGraphRootIdentifiers[0];
+
+		const FPluginReferenceNodeInfo& NodeInfo = ReferencerNodeInfos[FirstGraphRootIdentifier];
+		RootNode = CreatePluginReferenceNode();
+		RootNode->SetupPluginReferenceNode(CurrentGraphRootOrigin, NodeInfo.Identifier, NodeInfo.Plugin, !bIsCompactMode);
+
+		// Show referencers
+		RecursivelyCreateNodes(true, FirstGraphRootIdentifier, CurrentGraphRootOrigin, FirstGraphRootIdentifier, RootNode, ReferencerNodeInfos, 0, SettingsReferencerMaxDepth, /*bIsRoot*/ true);
+
+		// Show dependencies
+		RecursivelyCreateNodes(false, FirstGraphRootIdentifier, CurrentGraphRootOrigin, FirstGraphRootIdentifier, RootNode, DependencyNodeInfos, 0, SettingsDependencyMaxDepth, /*bIsRoot*/ true);
+	}
+
+	return RootNode;
+}
+
 UEdGraphNode_PluginReference* UEdGraph_PluginReferenceViewer::CreatePluginReferenceNode()
 {
 	const bool bSelectNewNode = false;
diff --git a/Engine/Plugins/Experimental/PluginReferenceViewer/Source/PluginReferenceViewer/Private/EdGraph_PluginReferenceViewer.h b/Engine/Plugins/Experimental/PluginReferenceViewer/Source/PluginReferenceViewer/Private/EdGraph_PluginReferenceViewer.h
--- a/Engine/Plugins/Experimental/PluginReferenceViewer/Source/PluginReferenceViewer/Private/EdGraph_PluginReferenceViewer.h
+++ b/Engine/Plugins/Experimental/PluginReferenceViewer/Source/PluginReferenceViewer/Private/EdGraph_PluginReferenceViewer.h
@@ -26,11 +26,6 @@ struct FPluginIdentifier
 	{
 	}
 
-	FPluginIdentifier(const FString& InName)
-		: Name(InName)
-	{
-	}
-
 	friend inline bool operator==(const FPluginIdentifier& A, const FPluginIdentifier& B)
 	{
 		return A.Name == B.Name;
@@ -40,6 +35,21 @@ struct FPluginIdentifier
 	{
 		return GetTypeHash(Key.Name);
 	}
+
+	bool IsValid() const
+	{
+		return Name != NAME_None;
+	}
+
+	FString ToString() const
+	{
+		return Name.ToString();
+	}
+
+	static FPluginIdentifier FromString(const FString& String)
+	{
+		return FPluginIdentifier(FName(String));
+	}
 };
 
 struct FPluginDependsNode
@@ -88,6 +98,8 @@ public:
 	/** Set reference viewer to focus on these assets */
 	void SetGraphRoot(const TArray<FPluginIdentifier>& GraphRootIdentifiers, const FIntPoint& GraphRootOrigin = FIntPoint(ForceInitToZero));
 
+	const TArray<FPluginIdentifier>& GetCurrentGraphRootIdentifiers() const;
+
 	UEdGraphNode_PluginReference* ConstructNodes(const TArray<FPluginIdentifier>& GraphRootIdentifiers, const FIntPoint& GraphRootOrigin);
 
 	UEdGraphNode_PluginReference* RecursivelyCreateNodes(bool bInReferencers, const FPluginIdentifier& InPluginId, const FIntPoint& InNodeLoc, const FPluginIdentifier& InParentId, UEdGraphNode_PluginReference* InParentNode, TMap<FPluginIdentifier, FPluginReferenceNodeInfo>& InNodeInfos, int32 InCurrentDepth, int32 InMaxDepth, bool bIsRoot = false);
@@ -97,7 +109,7 @@ public:
 	void GetSortedLinks(const TArray<FPluginIdentifier>& Identifiers, bool bReferencers, TMap<FPluginIdentifier, EPluginReferencePinCategory>& OutLinks) const;
 
 	/** Force the graph to rebuild */
-	class UEdGraphNode_PluginReference* RebuildGraph();
+	UEdGraphNode_PluginReference* RebuildGraph();
 
 	UEdGraphNode_PluginReference* CreatePluginReferenceNode();
 	void RemoveAllNodes();
@@ -122,6 +134,8 @@ private:
 	void GetPluginDependencies(const FPluginIdentifier& PluginIdentifier, TArray<FPluginDependency>& OutDependencies) const;
 	void GetPluginReferencers(const FPluginIdentifier& PluginIdentifier, TArray<FPluginDependency>& OutReferencers) const;
 
+	UEdGraphNode_PluginReference* RefilterGraph();
+
 private:
 	/** Editor for this pool */
 	TWeakPtr<SPluginReferenceViewer> PluginReferenceViewer;
diff --git a/Engine/Plugins/Experimental/PluginReferenceViewer/Source/PluginReferenceViewer/Private/HistoryManager.cpp b/Engine/Plugins/Experimental/PluginReferenceViewer/Source/PluginReferenceViewer/Private/HistoryManager.cpp
new file mode 100644
--- /dev/null
+++ b/Engine/Plugins/Experimental/PluginReferenceViewer/Source/PluginReferenceViewer/Private/HistoryManager.cpp
@@ -0,0 +1,159 @@
+// Copyright Epic Games, Inc. All Rights Reserved.
+
+#include "HistoryManager.h"
+#include "EdGraph_PluginReferenceViewer.h"
+
+FPluginReferenceViewerHistoryManager::FPluginReferenceViewerHistoryManager()
+{
+	CurrentHistoryIndex = 0;
+	MaxHistoryEntries = 300;
+}
+
+void FPluginReferenceViewerHistoryManager::SetOnApplyHistoryData(const FOnApplyHistoryData& InOnApplyHistoryData)
+{
+	OnApplyHistoryData = InOnApplyHistoryData;
+}
+
+void FPluginReferenceViewerHistoryManager::SetOnUpdateHistoryData(const FOnUpdateHistoryData& InOnUpdateHistoryData)
+{
+	OnUpdateHistoryData = InOnUpdateHistoryData;
+}
+
+bool FPluginReferenceViewerHistoryManager::GoBack()
+{
+	if ( CanGoBack() )
+	{
+		// Update the current history data
+		UpdateCurrentHistoryData();
+
+		// if its possible to go back, decrement the index we are at
+		--CurrentHistoryIndex;
+
+		// Update the owner
+		ApplyCurrentHistoryData();
+
+		return true;
+	}
+
+	return false;
+}
+
+bool FPluginReferenceViewerHistoryManager::GoForward()
+{
+	if ( CanGoForward() )
+	{
+		// Update the current history data
+		UpdateCurrentHistoryData();
+
+		// if its possible to go forward, increment the index we are at
+		++CurrentHistoryIndex;
+
+		// Update the owner
+		ApplyCurrentHistoryData();
+
+		return true;
+	}
+
+	return false;
+}
+
+void FPluginReferenceViewerHistoryManager::AddHistoryData()
+{
+	if (HistoryData.Num() == 0)
+	{
+		// History added to the beginning
+		HistoryData.Add(FPluginReferenceViewerHistoryData());
+		CurrentHistoryIndex = 0;
+	}
+	else if (CurrentHistoryIndex == HistoryData.Num() - 1)
+	{
+		// History added to the end
+		if (HistoryData.Num() == MaxHistoryEntries)
+		{
+			// If max history entries has been reached
+			// remove the oldest history
+			HistoryData.RemoveAt(0);
+		}
+		HistoryData.Add(FPluginReferenceViewerHistoryData());
+		// Current history index is the last index in the list
+		CurrentHistoryIndex = HistoryData.Num() - 1;
+	}
+	else
+	{
+		// History added to the middle
+		// clear out all history after the current history index.
+		HistoryData.RemoveAt(CurrentHistoryIndex + 1, HistoryData.Num() - (CurrentHistoryIndex + 1));
+		HistoryData.Add(FPluginReferenceViewerHistoryData());
+		// Current history index is the last index in the list
+		CurrentHistoryIndex = HistoryData.Num() - 1;
+	}
+
+	// Update the current history data
+	UpdateCurrentHistoryData();
+}
+
+void FPluginReferenceViewerHistoryManager::UpdateHistoryData()
+{
+	// Update the current history data
+	UpdateCurrentHistoryData();
+}
+
+bool FPluginReferenceViewerHistoryManager::CanGoForward() const
+{
+	// User can go forward if there are items in the history data list, 
+	// and the current history index isn't the last index in the list
+	return HistoryData.Num() > 0 && CurrentHistoryIndex < HistoryData.Num() - 1;
+}
+
+bool FPluginReferenceViewerHistoryManager::CanGoBack() const
+{
+	// User can go back if there are items in the history data list,
+	// and the current history index isn't the first index in the list
+	return HistoryData.Num() > 0 && CurrentHistoryIndex > 0;
+}
+
+FText FPluginReferenceViewerHistoryManager::GetBackDesc() const
+{
+	if ( CanGoBack() )
+	{
+		return HistoryData[CurrentHistoryIndex - 1].HistoryDesc;
+	}
+	return FText::GetEmpty();
+}
+
+FText FPluginReferenceViewerHistoryManager::GetForwardDesc() const
+{
+	if ( CanGoForward() )
+	{
+		return HistoryData[CurrentHistoryIndex + 1].HistoryDesc;
+	}
+	return FText::GetEmpty();
+}
+
+void FPluginReferenceViewerHistoryManager::ApplyCurrentHistoryData()
+{
+	if ( CurrentHistoryIndex >= 0 && CurrentHistoryIndex < HistoryData.Num())
+	{
+		OnApplyHistoryData.ExecuteIfBound( HistoryData[CurrentHistoryIndex] );
+	}
+}
+
+void FPluginReferenceViewerHistoryManager::UpdateCurrentHistoryData()
+{
+	if ( CurrentHistoryIndex >= 0 && CurrentHistoryIndex < HistoryData.Num())
+	{
+		OnUpdateHistoryData.ExecuteIfBound( HistoryData[CurrentHistoryIndex] );
+	}
+}
+
+void FPluginReferenceViewerHistoryManager::ExecuteJumpToHistory(int32 HistoryIndex)
+{
+	if (HistoryIndex >= 0 && HistoryIndex < HistoryData.Num())
+	{
+		// if the history index is valid, set the current history index to the history index requested by the user
+		CurrentHistoryIndex = HistoryIndex;
+
+		// Update the owner
+		ApplyCurrentHistoryData();
+	}
+}
\ No newline at end of file
diff --git a/Engine/Plugins/Experimental/PluginReferenceViewer/Source/PluginReferenceViewer/Private/HistoryManager.h b/Engine/Plugins/Experimental/PluginReferenceViewer/Source/PluginReferenceViewer/Private/HistoryManager.h
new file mode 100644
--- /dev/null
+++ b/Engine/Plugins/Experimental/PluginReferenceViewer/Source/PluginReferenceViewer/Private/HistoryManager.h
@@ -0,0 +1,89 @@
+// Copyright Epic Games, Inc. All Rights Reserved.
+
+#pragma once
+
+#include "Internationalization/Text.h"
+
+struct FPluginIdentifier;
+
+class FMenuBuilder;
+
+/** The history data object, storing all important history data */
+struct FPluginReferenceViewerHistoryData
+{
+	/** History Description */
+	FText HistoryDesc;
+
+	/** The list of package names to serve as the root */
+	TArray<FPluginIdentifier> Identifiers;
+};
+
+/** The delegate for when history data should be applied */
+DECLARE_DELEGATE_OneParam(FOnApplyHistoryData, const FPluginReferenceViewerHistoryData& /*Data*/);
+
+/** The delegate for when history data should be updated */
+DECLARE_DELEGATE_OneParam(FOnUpdateHistoryData, FPluginReferenceViewerHistoryData& /*Data*/);
+
+/** The class responsible for managing all content browser history */
+class FPluginReferenceViewerHistoryManager
+{
+public:
+	/** Constructor */
+	FPluginReferenceViewerHistoryManager();
+
+	/** Set the delegate for applying history data */
+	void SetOnApplyHistoryData(const FOnApplyHistoryData& InOnApplyHistoryData);
+
+	/** Set the delegate for updating history data */
+	void SetOnUpdateHistoryData(const FOnUpdateHistoryData& InOnUpdateHistoryData);
+
+	/** Goes back one history snapshot and returns the history data at that snapshot */
+	bool GoBack();
+
+	/** Goes forward one history snapshot and returns the history data at that snapshot */
+	bool GoForward();
+
+	/** Stores new history data.  Called when creating a history snapshot */
+	void AddHistoryData();
+
+	/** Triggers an update for the current history data. This is typically done right before changing the history. */
+	void UpdateHistoryData();
+
+	/** Determines if a user can go forward in history */
+	bool CanGoForward() const;
+
+	/** Determines if a user can go back in history */
+	bool CanGoBack() const;
+
+	/** Gets the description of the previous history entry */
+	FText GetBackDesc() const;
+
+	/** Gets the description of the next history entry */
+	FText GetForwardDesc() const;
+
+private:
+	/** Notifies the owner to update to the state described by the current history data */
+	void ApplyCurrentHistoryData();
+
+	/** Notifies the owner to update the current history data */
+	void UpdateCurrentHistoryData();
+
+	/** Handler for when a history item is chosen in the GetAvailableHistroyMenuItems list */
+	void ExecuteJumpToHistory(int32 HistoryIndex);
+
+private:
+	/** The delegate for when history data should be applied */
+	FOnApplyHistoryData OnApplyHistoryData;
+
+	/** The delegate for when history data should be updated */
+	FOnUpdateHistoryData OnUpdateHistoryData;
+
+	/** A list of history snapshots */
+	TArray<FPluginReferenceViewerHistoryData> HistoryData;
+
+	/** The current history index the user is at (changes when the user goes back,forward, or history snapshots are taken) */
+	int32 CurrentHistoryIndex;
+
+	/** Max number of history items that can be stored.  Once the max is reached, the oldest history item is removed */
+	int32 MaxHistoryEntries;
+};
diff --git a/Engine/Plugins/Experimental/PluginReferenceViewer/Source/PluginReferenceViewer/Private/PluginReferenceViewerModule.cpp b/Engine/Plugins/Experimental/PluginReferenceViewer/Source/PluginReferenceViewer/Private/PluginReferenceViewerModule.cpp
--- a/Engine/Plugins/Experimental/PluginReferenceViewer/Source/PluginReferenceViewer/Private/PluginReferenceViewerModule.cpp
+++ b/Engine/Plugins/Experimental/PluginReferenceViewer/Source/PluginReferenceViewer/Private/PluginReferenceViewerModule.cpp
@@ -75,7 +75,7 @@ void FPluginReferenceViewerModule::OpenPluginReferenceViewerUI(const TSharedRef<
 	if (TSharedPtr<SDockTab> NewTab = FGlobalTabmanager::Get()->TryInvokeTab(PluginReferenceViewerTabName))
 	{
 		TSharedRef<SPluginReferenceViewer> PluginReferenceViewer = StaticCastSharedRef<SPluginReferenceViewer>(NewTab->GetContent());
-		PluginReferenceViewer->SetGraphRoot(NewGraphRootNames);
+		PluginReferenceViewer->SetGraphRootIdentifiers(NewGraphRootNames);
 	}
 }
 
diff --git a/Engine/Plugins/Experimental/PluginReferenceViewer/Source/PluginReferenceViewer/Private/SPluginReferenceViewer.cpp b/Engine/Plugins/Experimental/PluginReferenceViewer/Source/PluginReferenceViewer/Private/SPluginReferenceViewer.cpp
--- a/Engine/Plugins/Experimental/PluginReferenceViewer/Source/PluginReferenceViewer/Private/SPluginReferenceViewer.cpp
+++ b/Engine/Plugins/Experimental/PluginReferenceViewer/Source/PluginReferenceViewer/Private/SPluginReferenceViewer.cpp
@@ -8,10 +8,12 @@
 #include "Features/IModularFeatures.h"
 #include "Features/IPluginsEditorFeature.h"
 #include "Interfaces/IPluginManager.h"
-#include "PluginReferenceViewerSchema.h"
 #include "PluginReferenceViewerCommands.h"
+#include "PluginReferenceViewerSchema.h"
+#include "PluginReferenceViewerStyle.h"
 #include "ToolMenus.h"
 #include "Widgets/Input/SSpinBox.h"
+#include "Widgets/Input/SEditableTextBox.h"
 #include "Widgets/Layout/SBorder.h"
 
 #define LOCTEXT_NAMESPACE "PluginReferenceViewer"
@@ -34,6 +36,10 @@ void SPluginReferenceViewer::Construct(const FArguments& InArgs)
 
 	RegisterActions();
 
+	// Set up the history manager
+	HistoryManager.SetOnApplyHistoryData(FOnApplyHistoryData::CreateSP(this, &SPluginReferenceViewer::OnApplyHistoryData));
+	HistoryManager.SetOnUpdateHistoryData(FOnUpdateHistoryData::CreateSP(this, &SPluginReferenceViewer::OnUpdateHistoryData));
+
 	// Create the graph
 	GraphObj = NewObject<UEdGraph_PluginReferenceViewer>();
 	GraphObj->Schema = UPluginReferenceViewerSchema::StaticClass();
@@ -48,7 +54,9 @@ void SPluginReferenceViewer::Construct(const FArguments& InArgs)
 		.AdditionalCommands(PluginReferenceViewerActions)
 		.GraphToEdit(GraphObj)
 		.GraphEvents(GraphEvents)
-		.ShowGraphStateOverlay(false);
+		.ShowGraphStateOverlay(false)
+		.OnNavigateHistoryBack(FSimpleDelegate::CreateSP(this, &SPluginReferenceViewer::GraphNavigateHistoryBack))
+		.OnNavigateHistoryForward(FSimpleDelegate::CreateSP(this, &SPluginReferenceViewer::GraphNavigateHistoryForward));
 
 	GraphObj->CachePluginDependencies(IPluginManager::Get().GetDiscoveredPlugins());
 
@@ -56,6 +64,7 @@ void SPluginReferenceViewer::Construct(const FArguments& InArgs)
 	[
 		SNew(SVerticalBox)
 
+		// Path and history
 		+SVerticalBox::Slot()
 		.AutoHeight()
 		[
@@ -70,6 +79,25 @@ void SPluginReferenceViewer::Construct(const FArguments& InArgs)
 				[
 					MakeToolBar()
 				]
+
+				// Path
+				+SHorizontalBox::Slot()
+				.Padding(0, 0, 4, 0)
+				.FillWidth(1.f)
+				.VAlign(VAlign_Center)
+				[
+					SNew(SBorder)
+					.BorderImage( FAppStyle::GetBrush("ToolPanel.GroupBorder") )
+					[
+						SNew(SEditableTextBox)
+						.Text(this, &SPluginReferenceViewer::GetAddressBarText)
+						.OnTextCommitted(this, &SPluginReferenceViewer::OnAddressBarTextCommitted)
+						.OnTextChanged(this, &SPluginReferenceViewer::OnAddressBarTextChanged)
+						.SelectAllTextWhenFocused(true)
+						.SelectAllTextOnCommit(true)
+						.Style(FAppStyle::Get(), "ReferenceViewer.PathText")
+					]
+				]
 			]
 		]
 
@@ -204,13 +232,18 @@ void SPluginReferenceViewer::Construct(const FArguments& InArgs)
 	];
 }
 
-void SPluginReferenceViewer::SetGraphRoot(const TArray<FPluginIdentifier>& GraphRootIdentifiers)
+void SPluginReferenceViewer::SetGraphRootIdentifiers(const TArray<FPluginIdentifier>& NewGraphRootIdentifiers)
 {
-	GraphObj->SetGraphRoot(GraphRootIdentifiers);
+	GraphObj->SetGraphRoot(NewGraphRootIdentifiers);
 
 	GraphObj->RebuildGraph();
 
 	GraphEditorPtr->ZoomToFit(false);
+
+	// Set the initial history data
+	HistoryManager.AddHistoryData();
+
+	TemporaryPathBeingEdited = NewGraphRootIdentifiers.Num() > 0 ? FText() : FText(LOCTEXT("NoPluginsFound", "No Plugins Found"));
 }
 
 void SPluginReferenceViewer::OnOpenPluginProperties()
@@ -237,6 +270,37 @@ void SPluginReferenceViewer::OnOpenPluginProperties()
 	}
 }
 
+void SPluginReferenceViewer::OnApplyHistoryData(const FPluginReferenceViewerHistoryData& History)
+{
+	if (GraphObj)
+	{
+		GraphObj->SetGraphRoot(History.Identifiers);
+		UEdGraphNode_PluginReference* NewRootNode = GraphObj->RebuildGraph();
+
+		if (NewRootNode && ensure(GraphEditorPtr.IsValid()))
+		{
+			GraphEditorPtr->SetNodeSelection(NewRootNode, true);
+		}
+
+		TemporaryPathBeingEdited = FText();
+	}
+}
+
+void SPluginReferenceViewer::OnUpdateHistoryData(FPluginReferenceViewerHistoryData& HistoryData) const
+{
+	if (GraphObj)
+	{
+		const TArray<FPluginIdentifier>& CurrentGraphRootIdentifiers = GraphObj->GetCurrentGraphRootIdentifiers();
+		HistoryData.HistoryDesc = GetAddressBarText();
+		HistoryData.Identifiers = CurrentGraphRootIdentifiers;
+	}
+	else
+	{
+		HistoryData.HistoryDesc = FText::GetEmpty();
+		HistoryData.Identifiers.Empty();
+	}
+}
+
 bool SPluginReferenceViewer::HasAtLeastOneRealNodeSelected()
 {
 	return true;
@@ -256,6 +320,28 @@ TSharedRef<SWidget> SPluginReferenceViewer::MakeToolBar()
 {
 	FToolBarBuilder ToolBarBuilder(PluginReferenceViewerActions, FMultiBoxCustomization::None, TSharedPtr<FExtender>(), true);
 
+	ToolBarBuilder.AddToolBarButton(
+		FUIAction(
+			FExecuteAction::CreateSP(this, &SPluginReferenceViewer::BackClicked),
+			FCanExecuteAction::CreateSP(this, &SPluginReferenceViewer::IsBackEnabled)
+		),
+		NAME_None,
+		TAttribute<FText>(),
+		TAttribute<FText>::CreateSP(this, &SPluginReferenceViewer::GetHistoryBackTooltip),
+		FSlateIcon(FPluginReferenceViewerStyle::Get().GetStyleSetName(), "Icons.ArrowLeft"));
+
+	ToolBarBuilder.AddToolBarButton(
+		FUIAction(
+			FExecuteAction::CreateSP(this, &SPluginReferenceViewer::ForwardClicked),
+			FCanExecuteAction::CreateSP(this, &SPluginReferenceViewer::IsForwardEnabled)
+		),
+		NAME_None,
+		TAttribute<FText>(),
+		TAttribute<FText>::CreateSP(this, &SPluginReferenceViewer::GetHistoryForwardTooltip),
+		FSlateIcon(FPluginReferenceViewerStyle::Get().GetStyleSetName(), "Icons.ArrowRight"));
+
+	ToolBarBuilder.AddSeparator();
+
 	ToolBarBuilder.AddComboButton(
 		FUIAction(),
 		FOnGetContent::CreateSP(this, &SPluginReferenceViewer::GetShowMenuContent),
@@ -285,6 +371,115 @@ void SPluginReferenceViewer::RebuildGraph()
 	GraphObj->RebuildGraph();
 }
 
+bool SPluginReferenceViewer::IsBackEnabled() const
+{
+	return HistoryManager.CanGoBack();
+}
+
+bool SPluginReferenceViewer::IsForwardEnabled() const
+{
+	return HistoryManager.CanGoForward();
+}
+
+void SPluginReferenceViewer::BackClicked()
+{
+	HistoryManager.GoBack();
+}
+
+void SPluginReferenceViewer::ForwardClicked()
+{
+	HistoryManager.GoForward();
+}
+
+void SPluginReferenceViewer::RefreshClicked()
+{
+	RebuildGraph();
+	ZoomToFit();
+}
+
+void SPluginReferenceViewer::GraphNavigateHistoryBack()
+{
+	BackClicked();
+}
+
+void SPluginReferenceViewer::GraphNavigateHistoryForward()
+{
+	ForwardClicked();
+}
+
+FText SPluginReferenceViewer::GetHistoryBackTooltip() const
+{
+	if (HistoryManager.CanGoBack())
+	{
+		return FText::Format(LOCTEXT("HistoryBackTooltip", "Back to {0}"), HistoryManager.GetBackDesc());
+	}
+	return FText::GetEmpty();
+}
+
+FText SPluginReferenceViewer::GetHistoryForwardTooltip() const
+{
+	if (HistoryManager.CanGoForward())
+	{
+		return FText::Format(LOCTEXT("HistoryForwardTooltip", "Forward to {0}"), HistoryManager.GetForwardDesc());
+	}
+	return FText::GetEmpty();
+}
+
+FText SPluginReferenceViewer::GetAddressBarText() const
+{
+	if (GraphObj)
+	{
+		if (TemporaryPathBeingEdited.IsEmpty())
+		{
+			const TArray<FPluginIdentifier>& CurrentGraphRootIdentifiers = GraphObj->GetCurrentGraphRootIdentifiers();
+			if (CurrentGraphRootIdentifiers.Num() == 1)
+			{
+				TSharedPtr<const IPlugin> Plugin = IPluginManager::Get().FindPlugin(CurrentGraphRootIdentifiers[0].ToString());
+				FString PluginPath = Plugin->GetBaseDir();
+				FPaths::MakePathRelativeTo(PluginPath, *FPaths::RootDir());
+
+				return FText::FromString(PluginPath);
+			}
+			else if (CurrentGraphRootIdentifiers.Num() > 1)
+			{
+				return FText::Format(LOCTEXT("AddressBarMultiplePackagesText", "{0} and {1} others"), FText::FromString(CurrentGraphRootIdentifiers[0].ToString()), FText::AsNumber(CurrentGraphRootIdentifiers.Num()));
+			}
+		}
+		else
+		{
+			return TemporaryPathBeingEdited;
+		}
+	}
+
+	return FText();
+}
+
+void SPluginReferenceViewer::OnAddressBarTextCommitted(const FText& NewText, ETextCommit::Type CommitInfo)
+{
+	if (CommitInfo == ETextCommit::OnEnter)
+	{
+		TArray<FPluginIdentifier> NewPaths;
+
+		// Plugins are identified by their unique name. The rest of the path is not important but it is displayed
+		// in the address bar for added clarity.
+		FString PluginName = FPaths::GetCleanFilename(NewText.ToString());
+
+		TSharedPtr<const IPlugin> Plugin = IPluginManager::Get().FindPlugin(PluginName);
+		if (Plugin != nullptr)
+		{
+			FPluginIdentifier NewPath = FPluginIdentifier::FromString(Plugin->GetName());
+			NewPaths.Add(NewPath);
+		}
+
+		SetGraphRootIdentifiers(NewPaths);
+	}
+}
+
+void SPluginReferenceViewer::OnAddressBarTextChanged(const FText& NewText)
+{
+	TemporaryPathBeingEdited = NewText;
+}
+
 void SPluginReferenceViewer::ZoomToFit()
 {
 	GraphEditorPtr->ZoomToFit(true);
@@ -321,6 +516,9 @@ void SPluginReferenceViewer::ReCenterGraphOnNodes(const TSet<UObject*>& Nodes)
 			GraphEditorPtr->ClearSelectionSet();
 			GraphEditorPtr->SetNodeSelection(NewRootNode, true);
 		}
+
+		// Set the initial history data
+		HistoryManager.AddHistoryData();
 	}
 }
 
@@ -373,7 +571,7 @@ void SPluginReferenceViewer::RegisterActions()
 		FCanExecuteAction());
 }
 
-bool SPluginReferenceViewer::IsCompactModeChecked()
+bool SPluginReferenceViewer::IsCompactModeChecked() const
 {
 	return Settings.bIsCompactMode;
 }
diff --git a/Engine/Plugins/Experimental/PluginReferenceViewer/Source/PluginReferenceViewer/Private/SPluginReferenceViewer.h b/Engine/Plugins/Experimental/PluginReferenceViewer/Source/PluginReferenceViewer/Private/SPluginReferenceViewer.h
--- a/Engine/Plugins/Experimental/PluginReferenceViewer/Source/PluginReferenceViewer/Private/SPluginReferenceViewer.h
+++ b/Engine/Plugins/Experimental/PluginReferenceViewer/Source/PluginReferenceViewer/Private/SPluginReferenceViewer.h
@@ -5,6 +5,7 @@
 #include "Widgets/SCompoundWidget.h"
 #include "EdGraph_PluginReferenceViewer.h"
 #include "GraphEditor.h"
+#include "HistoryManager.h"
 
 //////////////////////////////////////////////////////////////////////////
 // SPluginReferenceViewer
@@ -28,11 +29,11 @@ public:
 public:
 	void Construct(const FArguments& InArgs);
 
-	void SetGraphRoot(const TArray<FPluginIdentifier>& GraphRootIdentifiers);
+	void SetGraphRootIdentifiers(const TArray<FPluginIdentifier>& GraphRootIdentifiers);
 
 	int32 GetSearchReferencerDepthCount() const;
 	int32 GetSearchDependencyDepthCount() const;
-	bool IsCompactModeChecked();
+	bool IsCompactModeChecked() const;
 	bool IsShowEnginePluginsChecked() const;
 	bool IsShowOptionalPluginsChecked() const;
 
@@ -50,11 +51,45 @@ private:
 	void ReCenterGraphOnNodes(const TSet<UObject*>& Nodes);
 	void RegisterActions();
 
+	/** Called when a node is double clicked */
 	void OnNodeDoubleClicked(UEdGraphNode* Node);
 
+	/** True if the user may use the history back button */
+	bool IsBackEnabled() const;
+
+	/** True if the user may use the history forward button */
+	bool IsForwardEnabled() const;
+
+	/** Handler for clicking the history back button */
+	void BackClicked();
+
+	/** Handler for clicking the history forward button */
+	void ForwardClicked();
+
+	/** Refresh the current view */
+	void RefreshClicked();
+
+	/** Handler for when the graph panel tells us to go back in history (like using the mouse thumb button) */
+	void GraphNavigateHistoryBack();
+
+	/** Handler for when the graph panel tells us to go forward in history (like using the mouse thumb button) */
+	void GraphNavigateHistoryForward();
+
+	/** Gets the tool tip text for the history back button */
+	FText GetHistoryBackTooltip() const;
+
+	/** Gets the tool tip text for the history forward button */
+	FText GetHistoryForwardTooltip() const;
+
+	FText GetAddressBarText() const;
+	void OnAddressBarTextCommitted(const FText& NewText, ETextCommit::Type CommitInfo);
+	void OnAddressBarTextChanged(const FText& NewText);
+
 	void OnCompactModeChanged();
 	void OnShowEnginePluginsChanged();
 	void OnShowOptionalPluginsChanged();
+	void OnApplyHistoryData(const FPluginReferenceViewerHistoryData& History);
+	void OnUpdateHistoryData(FPluginReferenceViewerHistoryData& HistoryData) const;
 
 	struct FPluginReferenceViewerSettings
 	{
@@ -66,6 +101,9 @@ private:
 	};
 
 private:
+	/** The manager that keeps track of history data for this browser */
+	FPluginReferenceViewerHistoryManager HistoryManager;
+
 	TSharedPtr<SGraphEditor> GraphEditorPtr;
 	TSharedPtr<FUICommandList> PluginReferenceViewerActions;
 	TSharedPtr<SWidget> ReferencerCountBox;
@@ -74,6 +112,9 @@ private:
 	UEdGraph_PluginReferenceViewer* GraphObj;
 	FPluginReferenceViewerSettings Settings;
 
+	/** The temporary copy of the path text when it is actively being edited. */
+	FText TemporaryPathBeingEdited;
+
 	TSharedPtr<FUICommandInfo> ShowEnginePlugins;
 
 	/** Used to delay graph rebuilding during spinbox slider interaction */
