PluginReferenceViewer changes
- Toggle to show duplicates plugins or just the first referenced instance
- Added icon with tooltip to state that a node is referenced multiple times
- Filtering path added to avoid full graph rebuild from certain view options

diff --git a/Engine/Plugins/Experimental/PluginReferenceViewer/Source/PluginReferenceViewer/Private/EdGraphNode_PluginReference.cpp b/Engine/Plugins/Experimental/PluginReferenceViewer/Source/PluginReferenceViewer/Private/EdGraphNode_PluginReference.cpp
--- a/Engine/Plugins/Experimental/PluginReferenceViewer/Source/PluginReferenceViewer/Private/EdGraphNode_PluginReference.cpp
+++ b/Engine/Plugins/Experimental/PluginReferenceViewer/Source/PluginReferenceViewer/Private/EdGraphNode_PluginReference.cpp
@@ -15,7 +15,9 @@
 
 UEdGraphNode_PluginReference::UEdGraphNode_PluginReference(const FObjectInitializer& ObjectInitializer)
 	: Super(ObjectInitializer)
+	, bAllowThumbnail(true)
 	, bIsEnginePlugin(false)
+	, bIsADuplicate(false)
 	, DependencyPin(nullptr)
 	, ReferencerPin(nullptr)
 {
@@ -26,7 +28,7 @@ FPluginIdentifier UEdGraphNode_PluginReference::GetIdentifier() const
 	return PluginIdentifier;
 }
 
-void UEdGraphNode_PluginReference::SetupPluginReferenceNode(const FIntPoint& InNodeLoc, const FPluginIdentifier InPluginIdentifier, const TSharedPtr<const IPlugin>& InPlugin, bool bInAllowThumbnail)
+void UEdGraphNode_PluginReference::SetupPluginReferenceNode(const FIntPoint& InNodeLoc, const FPluginIdentifier InPluginIdentifier, const TSharedPtr<const IPlugin>& InPlugin, bool bInAllowThumbnail, bool bInIsADuplicate)
 {
 	NodePosX = InNodeLoc.X;
 	NodePosY = InNodeLoc.Y;
@@ -38,6 +40,7 @@ void UEdGraphNode_PluginReference::SetupPluginReferenceNode(const FIntPoint& InN
 
 	bAllowThumbnail = bInAllowThumbnail;
 	bIsEnginePlugin = (InPlugin->GetType() == EPluginType::Engine) ? true : false;
+	bIsADuplicate = bInIsADuplicate;
 
 	CachedPlugin = InPlugin;
 
diff --git a/Engine/Plugins/Experimental/PluginReferenceViewer/Source/PluginReferenceViewer/Private/EdGraphNode_PluginReference.h b/Engine/Plugins/Experimental/PluginReferenceViewer/Source/PluginReferenceViewer/Private/EdGraphNode_PluginReference.h
--- a/Engine/Plugins/Experimental/PluginReferenceViewer/Source/PluginReferenceViewer/Private/EdGraphNode_PluginReference.h
+++ b/Engine/Plugins/Experimental/PluginReferenceViewer/Source/PluginReferenceViewer/Private/EdGraphNode_PluginReference.h
@@ -20,9 +20,11 @@ class UEdGraphNode_PluginReference : public UEdGraphNode
 
 	FPluginIdentifier GetIdentifier() const;
 
-	void SetupPluginReferenceNode(const FIntPoint& InNodeLoc, const FPluginIdentifier InPluginIdentifier, const TSharedPtr<const IPlugin>& InPlugin, bool bAllowThumbnail);
+	void SetupPluginReferenceNode(const FIntPoint& InNodeLoc, const FPluginIdentifier InPluginIdentifier, const TSharedPtr<const IPlugin>& InPlugin, bool bInAllowThumbnail, bool bInIsADuplicate);
 
 	bool AllowsThumbnail() const;
+	bool IsADuplicate() const { return bIsADuplicate; }
+
 	TSharedPtr<const IPlugin> GetPlugin() const;
 
 	/** Returns our owning graph */
@@ -47,6 +49,7 @@ private:
 
 	bool bAllowThumbnail;
 	bool bIsEnginePlugin;
+	bool bIsADuplicate;
 
 	TSharedPtr<const IPlugin> CachedPlugin;
 	FSlateIcon AssetBrush;
diff --git a/Engine/Plugins/Experimental/PluginReferenceViewer/Source/PluginReferenceViewer/Private/EdGraph_PluginReferenceViewer.cpp b/Engine/Plugins/Experimental/PluginReferenceViewer/Source/PluginReferenceViewer/Private/EdGraph_PluginReferenceViewer.cpp
--- a/Engine/Plugins/Experimental/PluginReferenceViewer/Source/PluginReferenceViewer/Private/EdGraph_PluginReferenceViewer.cpp
+++ b/Engine/Plugins/Experimental/PluginReferenceViewer/Source/PluginReferenceViewer/Private/EdGraph_PluginReferenceViewer.cpp
@@ -178,13 +178,15 @@ UEdGraphNode_PluginReference* UEdGraph_PluginReferenceViewer::RecursivelyCreateN
 	else
 	{
 		const bool bIsCompactMode = PluginReferenceViewer.Pin()->IsCompactModeChecked();
+		const bool bIsADuplicate = NodeInfo.Parents.Num() > 1;
 
 		NewNode = CreatePluginReferenceNode();
-		NewNode->SetupPluginReferenceNode(InNodeLoc, NodeInfo.Identifier, NodeInfo.Plugin, !bIsCompactMode);
+		NewNode->SetupPluginReferenceNode(InNodeLoc, NodeInfo.Identifier, NodeInfo.Plugin, !bIsCompactMode, bIsADuplicate);
 		NodeProvSize = NodeInfo.ProvisionSize(InParentId);
 	}
 
 	const bool bIsCompactMode = PluginReferenceViewer.Pin()->IsCompactModeChecked();
+	const bool bShowDuplicates = PluginReferenceViewer.Pin()->IsShowDuplicatesChecked();
 
 	bool bIsFirstOccurance = bIsRoot || NodeInfo.IsFirstParent(InParentId);
 	FIntPoint ChildLoc = InNodeLoc;
@@ -204,11 +206,15 @@ UEdGraphNode_PluginReference* UEdGraph_PluginReferenceViewer::RecursivelyCreateN
 			const TPair<FPluginIdentifier, EPluginReferencePinCategory>& Pair = InNodeInfos[InPluginId].Children[ChildIdx];
 
 			FPluginIdentifier ChildId = Pair.Key;
-			int32 ChildProvSize = 1;
+			int32 ChildProvSize = 0;
 			if (InNodeInfos[ChildId].IsFirstParent(InPluginId))
 			{
 				ChildProvSize = InNodeInfos[ChildId].ProvisionSize(InPluginId);
 			}
+			else if (bShowDuplicates)
+			{
+				ChildProvSize = 1;
+			}
 
 			// The provision size will always be at least 1 if it should be shown
 			if (ChildProvSize > 0)
@@ -325,18 +331,54 @@ UEdGraphNode_PluginReference* UEdGraph_PluginReferenceViewer::RefilterGraph()
 
 		const FPluginReferenceNodeInfo& NodeInfo = ReferencerNodeInfos[FirstGraphRootIdentifier];
 		RootNode = CreatePluginReferenceNode();
-		RootNode->SetupPluginReferenceNode(CurrentGraphRootOrigin, NodeInfo.Identifier, NodeInfo.Plugin, !bIsCompactMode);
+		RootNode->SetupPluginReferenceNode(CurrentGraphRootOrigin, NodeInfo.Identifier, NodeInfo.Plugin, !bIsCompactMode, false);
 
 		// Show referencers
+		RecursivelyFilterNodeInfos(FirstGraphRootIdentifier, ReferencerNodeInfos, 0, SettingsReferencerMaxDepth);
 		RecursivelyCreateNodes(true, FirstGraphRootIdentifier, CurrentGraphRootOrigin, FirstGraphRootIdentifier, RootNode, ReferencerNodeInfos, 0, SettingsReferencerMaxDepth, /*bIsRoot*/ true);
 
 		// Show dependencies
+		RecursivelyFilterNodeInfos(FirstGraphRootIdentifier, DependencyNodeInfos, 0, SettingsDependencyMaxDepth);
 		RecursivelyCreateNodes(false, FirstGraphRootIdentifier, CurrentGraphRootOrigin, FirstGraphRootIdentifier, RootNode, DependencyNodeInfos, 0, SettingsDependencyMaxDepth, /*bIsRoot*/ true);
 	}
 
 	return RootNode;
 }
 
+void UEdGraph_PluginReferenceViewer::RecursivelyFilterNodeInfos(const FPluginIdentifier& InPluginId, TMap<FPluginIdentifier, FPluginReferenceNodeInfo>& InNodeInfos, int32 InCurrentDepth, int32 InMaxDepth)
+{
+	const bool bShowDuplicates = PluginReferenceViewer.Pin()->IsShowDuplicatesChecked();
+
+	// Filters and Re-provisions the NodeInfo counts 
+	int32 NewProvisionSize = 0;
+
+	if (!ExceedsMaxSearchDepth(InCurrentDepth, InMaxDepth))
+	{
+		for (const TPair<FPluginIdentifier, EPluginReferencePinCategory>& Pair : InNodeInfos[InPluginId].Children)
+		{
+			FPluginIdentifier ChildId = Pair.Key;
+
+			int32 ChildProvSize = 0;
+			if (InNodeInfos[ChildId].IsFirstParent(InPluginId))
+			{
+				RecursivelyFilterNodeInfos(ChildId, InNodeInfos, InCurrentDepth + 1, InMaxDepth);
+				ChildProvSize = InNodeInfos[ChildId].ProvisionSize(InPluginId);
+			}
+			else if (bShowDuplicates)
+			{
+				ChildProvSize = 1;
+			}
+
+			if (ChildProvSize > 0)
+			{
+				NewProvisionSize += ChildProvSize;
+			}
+		}
+	}
+
+	InNodeInfos[InPluginId].ChildProvisionSize = NewProvisionSize > 0 ? NewProvisionSize : 1;
+}
+
 UEdGraphNode_PluginReference* UEdGraph_PluginReferenceViewer::CreatePluginReferenceNode()
 {
 	const bool bSelectNewNode = false;
@@ -407,4 +449,4 @@ void UEdGraph_PluginReferenceViewer::GetPluginReferencers(const FPluginIdentifie
 			}
 		}
 	}
-}
\ No newline at end of file
+}
diff --git a/Engine/Plugins/Experimental/PluginReferenceViewer/Source/PluginReferenceViewer/Private/EdGraph_PluginReferenceViewer.h b/Engine/Plugins/Experimental/PluginReferenceViewer/Source/PluginReferenceViewer/Private/EdGraph_PluginReferenceViewer.h
--- a/Engine/Plugins/Experimental/PluginReferenceViewer/Source/PluginReferenceViewer/Private/EdGraph_PluginReferenceViewer.h
+++ b/Engine/Plugins/Experimental/PluginReferenceViewer/Source/PluginReferenceViewer/Private/EdGraph_PluginReferenceViewer.h
@@ -106,6 +106,8 @@ public:
 
 	void RecursivelyPopulateNodeInfos(bool bInReferencers, const TArray<FPluginIdentifier>& Identifiers, TMap<FPluginIdentifier, FPluginReferenceNodeInfo>& InNodeInfos, int32 InCurrentDepth, int32 InMaxDepth);
 
+	void RecursivelyFilterNodeInfos(const FPluginIdentifier& InPluginId, TMap<FPluginIdentifier, FPluginReferenceNodeInfo>& InNodeInfos, int32 InCurrentDepth, int32 InMaxDepth);
+
 	void GetSortedLinks(const TArray<FPluginIdentifier>& Identifiers, bool bReferencers, TMap<FPluginIdentifier, EPluginReferencePinCategory>& OutLinks) const;
 
 	/** Force the graph to rebuild */
diff --git a/Engine/Plugins/Experimental/PluginReferenceViewer/Source/PluginReferenceViewer/Private/PluginReferenceViewerCommands.cpp b/Engine/Plugins/Experimental/PluginReferenceViewer/Source/PluginReferenceViewer/Private/PluginReferenceViewerCommands.cpp
--- a/Engine/Plugins/Experimental/PluginReferenceViewer/Source/PluginReferenceViewer/Private/PluginReferenceViewerCommands.cpp
+++ b/Engine/Plugins/Experimental/PluginReferenceViewer/Source/PluginReferenceViewer/Private/PluginReferenceViewerCommands.cpp
@@ -16,6 +16,7 @@ FPluginReferenceViewerCommands::FPluginReferenceViewerCommands() : TCommands<FPl
 void FPluginReferenceViewerCommands::RegisterCommands()
 {	
 	UI_COMMAND(CompactMode, "Compact Mode", "Toggles Compact View", EUserInterfaceActionType::ToggleButton, FInputChord(EKeys::V));
+	UI_COMMAND(ShowDuplicates, "Show Duplicates References", "Toggles visibility of duplicates references", EUserInterfaceActionType::ToggleButton, FInputChord(EKeys::U));
 	UI_COMMAND(ShowEnginePlugins, "Show Engine Plugins", "Toggles visibility of Engine Plugins", EUserInterfaceActionType::ToggleButton, FInputChord());
 	UI_COMMAND(ShowOptionalPlugins, "Show Optional Plugins", "Toggles visibility of Optional Plugin dependencies", EUserInterfaceActionType::ToggleButton, FInputChord());
 
diff --git a/Engine/Plugins/Experimental/PluginReferenceViewer/Source/PluginReferenceViewer/Private/PluginReferenceViewerCommands.h b/Engine/Plugins/Experimental/PluginReferenceViewer/Source/PluginReferenceViewer/Private/PluginReferenceViewerCommands.h
--- a/Engine/Plugins/Experimental/PluginReferenceViewer/Source/PluginReferenceViewer/Private/PluginReferenceViewerCommands.h
+++ b/Engine/Plugins/Experimental/PluginReferenceViewer/Source/PluginReferenceViewer/Private/PluginReferenceViewerCommands.h
@@ -15,6 +15,7 @@ public:
 	// End of TCommands<> interface
 
 	TSharedPtr<FUICommandInfo> CompactMode;
+	TSharedPtr<FUICommandInfo> ShowDuplicates;	
 	TSharedPtr<FUICommandInfo> ShowEnginePlugins;
 	TSharedPtr<FUICommandInfo> ShowOptionalPlugins;
 
diff --git a/Engine/Plugins/Experimental/PluginReferenceViewer/Source/PluginReferenceViewer/Private/SPluginReferenceNode.cpp b/Engine/Plugins/Experimental/PluginReferenceViewer/Source/PluginReferenceViewer/Private/SPluginReferenceNode.cpp
--- a/Engine/Plugins/Experimental/PluginReferenceViewer/Source/PluginReferenceViewer/Private/SPluginReferenceNode.cpp
+++ b/Engine/Plugins/Experimental/PluginReferenceViewer/Source/PluginReferenceViewer/Private/SPluginReferenceNode.cpp
@@ -81,7 +81,7 @@ void SPluginReferenceNode::UpdateGraphNode()
 
 	TSharedRef<SWidget> ThumbnailWidget = SNullWidget::NullWidget;
 	UEdGraphNode_PluginReference* RefGraphNode = CastChecked<UEdGraphNode_PluginReference>(GraphNode);
-	bool bIsADuplicate = /*RefGraphNode->IsADuplicate()*/false;
+	bool bIsADuplicate = RefGraphNode->IsADuplicate();
 
 	FLinearColor OpacityColor = /*RefGraphNode->GetIsFiltered() ? FLinearColor(1.0, 1.0, 1.0, 0.4) :*/FLinearColor::White;
 	
@@ -179,7 +179,7 @@ void SPluginReferenceNode::UpdateGraphNode()
 								[
 									SNew(SImage)
 									.Visibility(bIsADuplicate ? EVisibility::Visible : EVisibility::Hidden)
-									.ToolTipText(LOCTEXT("DuplicateAsset", "This asset is referenced multiple times. Only the first occurance shows its decendants."))
+									.ToolTipText(LOCTEXT("DuplicateAsset", "This plugin is referenced multiple times. Only the first occurance shows its descendants."))
 									.DesiredSizeOverride(FVector2D(12.0, 12.0))
 									.Image(FPluginReferenceViewerStyle::Get().GetBrush("Graph.Node.Duplicate"))
 									.ColorAndOpacity(FAppStyle::Get().GetColor("Colors.Foreground"))
diff --git a/Engine/Plugins/Experimental/PluginReferenceViewer/Source/PluginReferenceViewer/Private/SPluginReferenceViewer.cpp b/Engine/Plugins/Experimental/PluginReferenceViewer/Source/PluginReferenceViewer/Private/SPluginReferenceViewer.cpp
--- a/Engine/Plugins/Experimental/PluginReferenceViewer/Source/PluginReferenceViewer/Private/SPluginReferenceViewer.cpp
+++ b/Engine/Plugins/Experimental/PluginReferenceViewer/Source/PluginReferenceViewer/Private/SPluginReferenceViewer.cpp
@@ -357,12 +357,16 @@ TSharedRef<SWidget> SPluginReferenceViewer::GetShowMenuContent()
 {
 	FMenuBuilder MenuBuilder(true, PluginReferenceViewerActions);
 
-	MenuBuilder.BeginSection("ViewOptions", LOCTEXT("ViewOptions", "View Options"));
-	MenuBuilder.AddMenuEntry(FPluginReferenceViewerCommands::Get().CompactMode);
+	MenuBuilder.BeginSection("Reference Types", LOCTEXT("ReferenceTypes", "Reference Types"));
 	MenuBuilder.AddMenuEntry(FPluginReferenceViewerCommands::Get().ShowEnginePlugins);
 	MenuBuilder.AddMenuEntry(FPluginReferenceViewerCommands::Get().ShowOptionalPlugins);
 	MenuBuilder.EndSection();
 
+	MenuBuilder.BeginSection("ViewOptions", LOCTEXT("ViewOptions", "View Options"));
+	MenuBuilder.AddMenuEntry(FPluginReferenceViewerCommands::Get().CompactMode);
+	MenuBuilder.AddMenuEntry(FPluginReferenceViewerCommands::Get().ShowDuplicates);
+	MenuBuilder.EndSection();
+
 	return MenuBuilder.MakeWidget();
 }
 
@@ -543,6 +547,12 @@ void SPluginReferenceViewer::RegisterActions()
 		FCanExecuteAction(),
 		FIsActionChecked::CreateSP(this, &SPluginReferenceViewer::IsCompactModeChecked));
 
+	PluginReferenceViewerActions->MapAction(
+		FPluginReferenceViewerCommands::Get().ShowDuplicates,
+		FExecuteAction::CreateSP(this, &SPluginReferenceViewer::OnShowDuplicatesChanged),
+		FCanExecuteAction(),
+		FIsActionChecked::CreateSP(this, &SPluginReferenceViewer::IsShowDuplicatesChecked));
+
 	PluginReferenceViewerActions->MapAction(
 		FPluginReferenceViewerCommands::Get().ShowEnginePlugins,
 		FExecuteAction::CreateSP(this, &SPluginReferenceViewer::OnShowEnginePluginsChanged),
@@ -579,7 +589,18 @@ bool SPluginReferenceViewer::IsCompactModeChecked() const
 void SPluginReferenceViewer::OnCompactModeChanged()
 {
 	Settings.bIsCompactMode = !Settings.bIsCompactMode;
-	GraphObj->RebuildGraph();
+	GraphObj->RefilterGraph();
+}
+
+void SPluginReferenceViewer::OnShowDuplicatesChanged()
+{
+	Settings.bShowDuplicates = !Settings.bShowDuplicates;
+	GraphObj->RefilterGraph();
+}
+
+bool SPluginReferenceViewer::IsShowDuplicatesChecked() const
+{
+	return Settings.bShowDuplicates;
 }
 
 bool SPluginReferenceViewer::IsShowEnginePluginsChecked() const
diff --git a/Engine/Plugins/Experimental/PluginReferenceViewer/Source/PluginReferenceViewer/Private/SPluginReferenceViewer.h b/Engine/Plugins/Experimental/PluginReferenceViewer/Source/PluginReferenceViewer/Private/SPluginReferenceViewer.h
--- a/Engine/Plugins/Experimental/PluginReferenceViewer/Source/PluginReferenceViewer/Private/SPluginReferenceViewer.h
+++ b/Engine/Plugins/Experimental/PluginReferenceViewer/Source/PluginReferenceViewer/Private/SPluginReferenceViewer.h
@@ -34,6 +34,7 @@ public:
 	int32 GetSearchReferencerDepthCount() const;
 	int32 GetSearchDependencyDepthCount() const;
 	bool IsCompactModeChecked() const;
+	bool IsShowDuplicatesChecked() const;
 	bool IsShowEnginePluginsChecked() const;
 	bool IsShowOptionalPluginsChecked() const;
 
@@ -86,6 +87,7 @@ private:
 	void OnAddressBarTextChanged(const FText& NewText);
 
 	void OnCompactModeChanged();
+	void OnShowDuplicatesChanged();
 	void OnShowEnginePluginsChanged();
 	void OnShowOptionalPluginsChanged();
 	void OnApplyHistoryData(const FPluginReferenceViewerHistoryData& History);
@@ -96,6 +98,7 @@ private:
 		int32 MaxSearchReferencersDepth = 1;
 		int32 MaxSearchDependencyDepth = 1;
 		bool bIsCompactMode = false;
+		bool bShowDuplicates = true;
 		bool bShowEnginePlugins = true;
 		bool bShowOptionalPlugins = true;
 	};
