PluginReferenceViewer - Added a context menu item that launches the asset reference viewer showing the asset references across plugins. Extended the OpenReferenceViewerUI to allow setting of the plugin filters on launch.

diff --git a/Engine/Plugins/Editor/AssetManagerEditor/Source/AssetManagerEditor/Private/ReferenceViewer/SReferenceViewer.cpp b/Engine/Plugins/Editor/AssetManagerEditor/Source/AssetManagerEditor/Private/ReferenceViewer/SReferenceViewer.cpp
--- a/Engine/Plugins/Editor/AssetManagerEditor/Source/AssetManagerEditor/Private/ReferenceViewer/SReferenceViewer.cpp
+++ b/Engine/Plugins/Editor/AssetManagerEditor/Source/AssetManagerEditor/Private/ReferenceViewer/SReferenceViewer.cpp
@@ -617,6 +617,12 @@ void SReferenceViewer::SetGraphRootIdentifiers(const TArray<FAssetIdentifier>& N
 		UAssetManager::Get().UpdateManagementDatabase();
 	}
 
+	if (!ReferenceViewerParams.PluginFilter.IsEmpty())
+	{
+		Settings->SetEnablePluginFilter(true);
+		GraphObj->SetCurrentPluginFilter(ReferenceViewerParams.PluginFilter);
+	}
+
 	RebuildGraph();
 
 	UpdateIsPassingSearchFilterCallback();
diff --git a/Engine/Plugins/Experimental/PluginReferenceViewer/PluginReferenceViewer.uplugin b/Engine/Plugins/Experimental/PluginReferenceViewer/PluginReferenceViewer.uplugin
--- a/Engine/Plugins/Experimental/PluginReferenceViewer/PluginReferenceViewer.uplugin
+++ b/Engine/Plugins/Experimental/PluginReferenceViewer/PluginReferenceViewer.uplugin
@@ -24,6 +24,10 @@
 		}
 	],
 	"Plugins": [
+		{
+			"Name": "AssetManagerEditor",
+			"Enabled": true
+		},	
 		{
 			"Name": "PluginUtils",
 			"Enabled": true
diff --git a/Engine/Plugins/Experimental/PluginReferenceViewer/Source/PluginReferenceViewer/PluginReferenceViewer.Build.cs b/Engine/Plugins/Experimental/PluginReferenceViewer/Source/PluginReferenceViewer/PluginReferenceViewer.Build.cs
--- a/Engine/Plugins/Experimental/PluginReferenceViewer/Source/PluginReferenceViewer/PluginReferenceViewer.Build.cs
+++ b/Engine/Plugins/Experimental/PluginReferenceViewer/Source/PluginReferenceViewer/PluginReferenceViewer.Build.cs
@@ -24,6 +24,8 @@ public class PluginReferenceViewer : ModuleRules
 
 		PrivateDependencyModuleNames.AddRange(
 			new string[] {
+				"AssetManagerEditor",
+				"AssetRegistry",
 				"EditorFramework",
 				"PluginUtils",
 				"Projects",
diff --git a/Engine/Plugins/Experimental/PluginReferenceViewer/Source/PluginReferenceViewer/Private/EdGraph_PluginReferenceViewer.cpp b/Engine/Plugins/Experimental/PluginReferenceViewer/Source/PluginReferenceViewer/Private/EdGraph_PluginReferenceViewer.cpp
--- a/Engine/Plugins/Experimental/PluginReferenceViewer/Source/PluginReferenceViewer/Private/EdGraph_PluginReferenceViewer.cpp
+++ b/Engine/Plugins/Experimental/PluginReferenceViewer/Source/PluginReferenceViewer/Private/EdGraph_PluginReferenceViewer.cpp
@@ -2,6 +2,9 @@
 
 #include "EdGraph_PluginReferenceViewer.h"
 
+#include "AssetManagerEditorModule.h"
+#include "AssetRegistry/AssetRegistryModule.h"
+#include "AssetRegistry/AssetData.h"
 #include "EdGraphNode_PluginReference.h"
 #include "EdGraph/EdGraphPin.h"
 #include "Interfaces/IPluginManager.h"
@@ -10,14 +13,73 @@
 
 #include UE_INLINE_GENERATED_CPP_BY_NAME(EdGraph_PluginReferenceViewer)
 
-bool ExceedsMaxSearchDepth(int32 CurrentDepth, int32 MaxDepth)
+namespace UE::PluginReferenceViewer::Private
 {
-	const bool bIsWithinDepthLimits = MaxDepth > 0 && CurrentDepth < MaxDepth;
-	if (!bIsWithinDepthLimits)
+	bool ExceedsMaxSearchDepth(int32 CurrentDepth, int32 MaxDepth)
 	{
-		return true;
+		const bool bIsWithinDepthLimits = MaxDepth > 0 && CurrentDepth < MaxDepth;
+		if (!bIsWithinDepthLimits)
+		{
+			return true;
+		}
+		return false;
 	}
-	return false;
+
+	void FindPluginAssetReferences(const TSharedRef<IPlugin>& ParentPlugin, const TSharedRef<IPlugin>& ChildPlugin, TArray<FAssetIdentifier>& OutAssetIdentifiers)
+	{
+		IAssetRegistry& AssetRegistry = FModuleManager::LoadModuleChecked<FAssetRegistryModule>("AssetRegistry").Get();
+
+		const FString ChildMountedAssetPath = ChildPlugin->GetMountedAssetPath();
+
+		FARFilter Filter;
+		Filter.bRecursivePaths = true;
+		Filter.PackagePaths.Add(FName(ParentPlugin->GetMountedAssetPath()));
+
+		TArray<FAssetData> AssetsInPlugin;
+		AssetRegistry.GetAssets(Filter, AssetsInPlugin);
+
+		// Reusing variables to minimize allocations during large asset list iteration.
+		TArray<FAssetDependency> AssetDependencies;
+		FString AssetPath;
+
+		for (const FAssetData& AssetData : AssetsInPlugin)
+		{
+			const FAssetIdentifier AssetIdentifier(AssetData.PackageName);
+
+			AssetDependencies.Empty();
+			AssetRegistry.GetDependencies(AssetIdentifier, AssetDependencies);
+
+			for (const FAssetDependency& AssetDependency : AssetDependencies)
+			{
+				AssetDependency.AssetId.PackageName.ToString(AssetPath);
+
+				if (AssetPath.StartsWith(ChildMountedAssetPath))
+				{
+					OutAssetIdentifiers.Add(AssetIdentifier);
+					break;
+				}
+			}
+		}
+	};
+
+	void GetPluginAssets(const TSharedRef<IPlugin>& Plugin, TArray<FAssetIdentifier>& OutAssetIdentifiers)
+	{
+		IAssetRegistry& AssetRegistry = FModuleManager::LoadModuleChecked<FAssetRegistryModule>("AssetRegistry").Get();
+
+		FARFilter Filter;
+		Filter.bRecursivePaths = true;
+		Filter.PackagePaths.Add(FName(Plugin->GetMountedAssetPath()));
+
+		TArray<FAssetData> AssetsInPlugin;
+		AssetRegistry.GetAssets(Filter, AssetsInPlugin);
+
+		OutAssetIdentifiers.Reserve(AssetsInPlugin.Num());
+
+		for (int32 Index = 0; Index < AssetsInPlugin.Num(); ++Index)
+		{
+			OutAssetIdentifiers.Add(FAssetIdentifier(AssetsInPlugin[Index].PackageName));
+		}
+	};
 }
 
 FPluginReferenceNodeInfo::FPluginReferenceNodeInfo(const FPluginIdentifier& InIdentifier, bool bInReferencers)
@@ -56,6 +118,21 @@ const TArray<FPluginIdentifier>& UEdGraph_PluginReferenceViewer::GetCurrentGraph
 	return CurrentGraphRootIdentifiers;
 }
 
+void UEdGraph_PluginReferenceViewer::FindAssetReferencesAcrossPlugins(const FPluginIdentifier& From, const FPluginIdentifier& To, TArray<FAssetIdentifier>& OutAssetIdentifiers)
+{
+	const TSharedRef<IPlugin> FromPlugin = PluginMap.FindChecked(From);
+	const TSharedRef<IPlugin> ToPlugin = PluginMap.FindChecked(To);
+
+	UE::PluginReferenceViewer::Private::FindPluginAssetReferences(FromPlugin, ToPlugin, OutAssetIdentifiers);
+}
+
+void UEdGraph_PluginReferenceViewer::GetPluginAssets(const FPluginIdentifier& Plugin, TArray<FAssetIdentifier>& OutAssetIdentifiers)
+{
+	const TSharedRef<IPlugin> ForPlugin = PluginMap.FindChecked(Plugin);
+
+	UE::PluginReferenceViewer::Private::GetPluginAssets(ForPlugin, OutAssetIdentifiers);
+}
+
 void UEdGraph_PluginReferenceViewer::CachePluginDependencies(const TArray<TSharedRef<IPlugin>>& Plugins)
 {
 	PluginMap.Empty();
@@ -190,7 +267,7 @@ UEdGraphNode_PluginReference* UEdGraph_PluginReferenceViewer::RecursivelyCreateN
 
 	bool bIsFirstOccurance = bIsRoot || NodeInfo.IsFirstParent(InParentId);
 	FIntPoint ChildLoc = InNodeLoc;
-	if (!ExceedsMaxSearchDepth(InCurrentDepth, InMaxDepth) && bIsFirstOccurance) // Only expand the first parent
+	if (!UE::PluginReferenceViewer::Private::ExceedsMaxSearchDepth(InCurrentDepth, InMaxDepth) && bIsFirstOccurance) // Only expand the first parent
 	{
 		// position the children nodes
 		const int32 ColumnWidth = bIsCompactMode ? 500 : 800;
@@ -275,7 +352,7 @@ void UEdGraph_PluginReferenceViewer::RecursivelyPopulateNodeInfos(bool bInRefere
 	check(Identifiers.Num() > 0);
 	int32 ProvisionSize = 0;
 	const FPluginIdentifier PluginId = Identifiers[0];
-	if (!ExceedsMaxSearchDepth(InCurrentDepth, InMaxDepth))
+	if (!UE::PluginReferenceViewer::Private::ExceedsMaxSearchDepth(InCurrentDepth, InMaxDepth))
 	{
 		TMap<FPluginIdentifier, EPluginReferencePinCategory> ReferenceLinks;
 		GetSortedLinks(Identifiers, bInReferencers, ReferenceLinks);
@@ -352,7 +429,7 @@ void UEdGraph_PluginReferenceViewer::RecursivelyFilterNodeInfos(const FPluginIde
 	// Filters and Re-provisions the NodeInfo counts 
 	int32 NewProvisionSize = 0;
 
-	if (!ExceedsMaxSearchDepth(InCurrentDepth, InMaxDepth))
+	if (!UE::PluginReferenceViewer::Private::ExceedsMaxSearchDepth(InCurrentDepth, InMaxDepth))
 	{
 		for (const TPair<FPluginIdentifier, EPluginReferencePinCategory>& Pair : InNodeInfos[InPluginId].Children)
 		{
@@ -394,6 +471,16 @@ void UEdGraph_PluginReferenceViewer::RemoveAllNodes()
 	}
 }
 
+bool UEdGraph_PluginReferenceViewer::IsDependencyNode(const FPluginIdentifier& PluginIdentifier) const
+{
+	return DependencyNodeInfos.Contains(PluginIdentifier);
+}
+
+bool UEdGraph_PluginReferenceViewer::IsReferencerNode(const FPluginIdentifier& PluginIdentifier) const
+{
+	return ReferencerNodeInfos.Contains(PluginIdentifier);
+}
+
 void UEdGraph_PluginReferenceViewer::SetPluginReferenceViewer(TSharedPtr<SPluginReferenceViewer> InViewer)
 {
 	PluginReferenceViewer = InViewer;
@@ -449,4 +536,4 @@ void UEdGraph_PluginReferenceViewer::GetPluginReferencers(const FPluginIdentifie
 			}
 		}
 	}
-}
+}
\ No newline at end of file
diff --git a/Engine/Plugins/Experimental/PluginReferenceViewer/Source/PluginReferenceViewer/Private/EdGraph_PluginReferenceViewer.h b/Engine/Plugins/Experimental/PluginReferenceViewer/Source/PluginReferenceViewer/Private/EdGraph_PluginReferenceViewer.h
--- a/Engine/Plugins/Experimental/PluginReferenceViewer/Source/PluginReferenceViewer/Private/EdGraph_PluginReferenceViewer.h
+++ b/Engine/Plugins/Experimental/PluginReferenceViewer/Source/PluginReferenceViewer/Private/EdGraph_PluginReferenceViewer.h
@@ -11,6 +11,7 @@
 class IPlugin;
 class UEdGraphNode_PluginReference;
 class SPluginReferenceViewer;
+struct FAssetIdentifier;
 
 struct FPluginIdentifier
 {
@@ -100,6 +101,17 @@ public:
 
 	const TArray<FPluginIdentifier>& GetCurrentGraphRootIdentifiers() const;
 
+	/**
+	 * Finds all plugin assets that have dependencies on another plugins assets
+	 *
+	 * @param From the plugin whose assets we want to find references from
+	 * @param To the plugin whose assets we want to find dependencies to
+	 */
+	void FindAssetReferencesAcrossPlugins(const FPluginIdentifier& From, const FPluginIdentifier& To, TArray<FAssetIdentifier>& OutAssetIdentifiers);
+
+	/** Finds all plugin assets */
+	void GetPluginAssets(const FPluginIdentifier& Plugin, TArray<FAssetIdentifier>& OutAssetIdentifiers);
+
 	UEdGraphNode_PluginReference* ConstructNodes(const TArray<FPluginIdentifier>& GraphRootIdentifiers, const FIntPoint& GraphRootOrigin);
 
 	UEdGraphNode_PluginReference* RecursivelyCreateNodes(bool bInReferencers, const FPluginIdentifier& InPluginId, const FIntPoint& InNodeLoc, const FPluginIdentifier& InParentId, UEdGraphNode_PluginReference* InParentNode, TMap<FPluginIdentifier, FPluginReferenceNodeInfo>& InNodeInfos, int32 InCurrentDepth, int32 InMaxDepth, bool bIsRoot = false);
@@ -116,6 +128,12 @@ public:
 	UEdGraphNode_PluginReference* CreatePluginReferenceNode();
 	void RemoveAllNodes();
 
+	/** Returns true if the plugin is a dependency of the graph root plugin. */
+	bool IsDependencyNode(const FPluginIdentifier& PluginIdentifier) const;
+
+	/** Returns true if the plugin is a referencer of the graph root plugin. */
+	bool IsReferencerNode(const FPluginIdentifier& PluginIdentifier) const;
+
 private:
 	struct FPluginDependency
 	{
diff --git a/Engine/Plugins/Experimental/PluginReferenceViewer/Source/PluginReferenceViewer/Private/PluginReferenceViewerSchema.cpp b/Engine/Plugins/Experimental/PluginReferenceViewer/Source/PluginReferenceViewer/Private/PluginReferenceViewerSchema.cpp
--- a/Engine/Plugins/Experimental/PluginReferenceViewer/Source/PluginReferenceViewer/Private/PluginReferenceViewerSchema.cpp
+++ b/Engine/Plugins/Experimental/PluginReferenceViewer/Source/PluginReferenceViewer/Private/PluginReferenceViewerSchema.cpp
@@ -2,7 +2,9 @@
 
 #include "PluginReferenceViewerSchema.h"
 
+#include "AssetManagerEditorModule.h"
 #include "ConnectionDrawingPolicy.h"
+#include "EdGraphNode_PluginReference.h"
 #include "PluginReferencePinCategory.h"
 #include "PluginReferenceViewerCommands.h"
 #include "SPluginReferenceViewer.h"
@@ -52,6 +54,79 @@ void UPluginReferenceViewerSchema::GetContextMenuActions(UToolMenu* Menu, UGraph
 		Section.AddMenuEntry(FPluginReferenceViewerCommands::Get().OpenPluginProperties);
 	}
 
+	{
+		FToolMenuSection& Section = Menu->AddSection(TEXT("AssetReferences"), NSLOCTEXT("PluginReferenceViewerSchema", "AssetReferencesSectionLabel", "Asset References"));
+
+		FText Title = NSLOCTEXT("PluginReferenceViewerSchema", "OpenReferenceViewer", "Open {0} Asset References...");
+
+		UEdGraph_PluginReferenceViewer* PluginReferenceViewer = CastChecked<UEdGraph_PluginReferenceViewer>(Context->Graph);
+		UEdGraphNode_PluginReference* ContextPluginNode = Cast<UEdGraphNode_PluginReference>(Context->Node);
+
+		// Opens the reference viewer for all assets belonging to the plugin.
+		check(!PluginReferenceViewer->GetCurrentGraphRootIdentifiers().IsEmpty());
+
+		FPluginIdentifier GraphRootPluginIdentifier = PluginReferenceViewer->GetCurrentGraphRootIdentifiers()[0];
+
+		if (ContextPluginNode->GetIdentifier() == GraphRootPluginIdentifier)
+		{
+			TArray<FAssetIdentifier> AssetIdentifiers;
+			PluginReferenceViewer->GetPluginAssets(ContextPluginNode->GetIdentifier(), AssetIdentifiers);
+
+			Section.AddEntry(FToolMenuEntry::InitMenuEntry(NAME_None,
+				FText::Format(Title, FText::AsNumber(AssetIdentifiers.Num())),
+				NSLOCTEXT("PluginReferenceViewerSchema", "PluginAssetReferences", "Launches the reference viewer showing the plugin asset references"),
+				FSlateIcon(),
+				FUIAction(FExecuteAction::CreateLambda(
+					[AssetIdentifiers]()
+					{
+						IAssetManagerEditorModule::Get().OpenReferenceViewerUI(AssetIdentifiers);
+					}),
+					FCanExecuteAction::CreateLambda([AssetCount = AssetIdentifiers.Num()]() { return AssetCount > 0; })
+				)
+			));
+		}
+		else
+		{
+			/*	
+			*	Opens the reference viewer for all assets that have references across plugins.
+			*		a) from the root plugin node to the selected dependency plugin node.
+			*		b) from the selected referencer plugin node to the root plugin node. 
+			*/
+			FText Tooltip = NSLOCTEXT("PluginReferenceViewerSchema", "AssetReferencesAcrossPlugins", "Launches the reference viewer showing the plugin asset references from '{0}' to '{1}'");
+			TArray<FAssetIdentifier> AssetIdentifiers;
+			FReferenceViewerParams ReferenceViewerParams;
+			ReferenceViewerParams.bShowPluginFilter = true;
+
+			if (PluginReferenceViewer->IsDependencyNode(ContextPluginNode->GetIdentifier()))
+			{
+				Tooltip = FText::Format(Tooltip, FText::FromName(GraphRootPluginIdentifier.Name), FText::FromName(ContextPluginNode->GetIdentifier().Name));
+
+				ReferenceViewerParams.PluginFilter.Add(ContextPluginNode->GetIdentifier().Name);
+				PluginReferenceViewer->FindAssetReferencesAcrossPlugins(GraphRootPluginIdentifier, ContextPluginNode->GetIdentifier(), AssetIdentifiers);
+			}
+			else
+			{
+				Tooltip = FText::Format(Tooltip, FText::FromName(ContextPluginNode->GetIdentifier().Name), FText::FromName(GraphRootPluginIdentifier.Name));
+
+				ReferenceViewerParams.PluginFilter.Add(GraphRootPluginIdentifier.Name);
+				PluginReferenceViewer->FindAssetReferencesAcrossPlugins(ContextPluginNode->GetIdentifier(), GraphRootPluginIdentifier, AssetIdentifiers);
+			}
+
+			Section.AddEntry(FToolMenuEntry::InitMenuEntry(NAME_None,
+				FText::Format(Title, FText::AsNumber(AssetIdentifiers.Num())),
+				Tooltip,
+				FSlateIcon(),
+				FUIAction(FExecuteAction::CreateLambda(
+					[AssetIdentifiers, ReferenceViewerParams]()
+					{
+						IAssetManagerEditorModule::Get().OpenReferenceViewerUI(AssetIdentifiers, ReferenceViewerParams);
+					}),
+					FCanExecuteAction::CreateLambda([AssetCount=AssetIdentifiers.Num()]() { return AssetCount > 0; })
+				)
+			));
+		}
+	}
+
 	{
 		FToolMenuSection& Section = Menu->AddSection(TEXT("Misc"), NSLOCTEXT("PluginReferenceViewerSchema", "MiscSectionLabel", "Misc"));
 		Section.AddMenuEntry(FPluginReferenceViewerCommands::Get().ZoomToFit);
diff --git a/Engine/Source/Runtime/CoreUObject/Public/AssetRegistry/AssetData.h b/Engine/Source/Runtime/CoreUObject/Public/AssetRegistry/AssetData.h
--- a/Engine/Source/Runtime/CoreUObject/Public/AssetRegistry/AssetData.h
+++ b/Engine/Source/Runtime/CoreUObject/Public/AssetRegistry/AssetData.h
@@ -1006,6 +1006,7 @@ struct FReferenceViewerParams
 		, FixAndHideSearchDepthLimit(0)
 		, FixAndHideSearchBreadthLimit(0)
 		, bShowCollectionFilter(true)
+		, bShowPluginFilter(true)
 		// Checkbox options
 		, bShowShowReferencesOptions(true)
 		, bShowShowSearchableNames(true)
@@ -1022,6 +1023,8 @@ struct FReferenceViewerParams
 	/* Whether to only display the References/Dependencies which match the text filter, if any. 
 	   If the optional is not set, don't change the current reference viewer's value. */
 	TOptional<bool> bShowFilteredPackagesOnly;
+	/** Whether to only display the References/Dependencies which match the plugin name filter */
+	TArray<FName> PluginFilter;
 	/* Compact mode allows to hide the thumbnail and minimize the space taken by the nodes. Useful when there are many dependencies to inspect, to keep the UI responsive. 
 	   If the optional is not set, don't change the current reference viewer's value. */
 	TOptional<bool> bCompactMode;
