Fix implicit const_cast
#rnx

diff --git a/Engine/Plugins/Experimental/PluginReferenceViewer/Source/PluginReferenceViewer/Private/PluginReferenceViewerSchema.cpp b/Engine/Plugins/Experimental/PluginReferenceViewer/Source/PluginReferenceViewer/Private/PluginReferenceViewerSchema.cpp
--- a/Engine/Plugins/Experimental/PluginReferenceViewer/Source/PluginReferenceViewer/Private/PluginReferenceViewerSchema.cpp
+++ b/Engine/Plugins/Experimental/PluginReferenceViewer/Source/PluginReferenceViewer/Private/PluginReferenceViewerSchema.cpp
@@ -59,8 +59,8 @@ void UPluginReferenceViewerSchema::GetContextMenuActions(UToolMenu* Menu, UGraph
 
 		FText Title = NSLOCTEXT("PluginReferenceViewerSchema", "OpenReferenceViewer", "Open {0} Asset References...");
 
-		UEdGraph_PluginReferenceViewer* PluginReferenceViewer = CastChecked<UEdGraph_PluginReferenceViewer>(Context->Graph);
-		UEdGraphNode_PluginReference* ContextPluginNode = Cast<UEdGraphNode_PluginReference>(Context->Node);
+		UEdGraph_PluginReferenceViewer* PluginReferenceViewer = const_cast<UEdGraph_PluginReferenceViewer*>(CastChecked<UEdGraph_PluginReferenceViewer>(Context->Graph));
+		const UEdGraphNode_PluginReference* ContextPluginNode = Cast<UEdGraphNode_PluginReference>(Context->Node);
 
 		// Opens the reference viewer for all assets belonging to the plugin.
 		check(!PluginReferenceViewer->GetCurrentGraphRootIdentifiers().IsEmpty());
