Add advanced info to Plugin Reference Viewer. Llimited extra info for right now, gives plugin content sizes, referenced content size, and reference counts.

diff --git a/Engine/Plugins/Experimental/PluginReferenceViewer/Source/PluginReferenceViewer/Private/EdGraph_PluginReferenceViewer.cpp b/Engine/Plugins/Experimental/PluginReferenceViewer/Source/PluginReferenceViewer/Private/EdGraph_PluginReferenceViewer.cpp
--- a/Engine/Plugins/Experimental/PluginReferenceViewer/Source/PluginReferenceViewer/Private/EdGraph_PluginReferenceViewer.cpp
+++ b/Engine/Plugins/Experimental/PluginReferenceViewer/Source/PluginReferenceViewer/Private/EdGraph_PluginReferenceViewer.cpp
@@ -5,6 +5,7 @@
 #include "AssetManagerEditorModule.h"
 #include "AssetRegistry/AssetRegistryModule.h"
 #include "AssetRegistry/AssetData.h"
+#include "Async/ParallelFor.h"
 #include "EdGraphNode_PluginReference.h"
 #include "EdGraph/EdGraphPin.h"
 #include "Interfaces/IPluginManager.h"
@@ -80,6 +81,28 @@ namespace UE::PluginReferenceViewer::Private
 			OutAssetIdentifiers.Add(FAssetIdentifier(AssetsInPlugin[Index].PackageName));
 		}
 	};
+
+	uint64 GetSizeOfPlugin(const TSharedRef<IPlugin> Plugin)
+	{
+		uint64 Size = 0;
+		TArray<FAssetData> PluginAssets;
+
+		IAssetRegistry& AssetRegistry = FModuleManager::LoadModuleChecked<FAssetRegistryModule>("AssetRegistry").Get();
+
+		// we want to only include on disk assets otherwise we will be stuck only on the game thread
+		AssetRegistry.GetAssetsByPath(FName(Plugin->GetMountedAssetPath()), PluginAssets, true, true);
+
+		for (const FAssetData& PluginAsset : PluginAssets)
+		{
+			if (TOptional<FAssetPackageData> PackageData = AssetRegistry.GetAssetPackageDataCopy(PluginAsset.PackageName))
+			{
+				Size += PackageData->DiskSize;
+			}
+		}
+
+		return Size;
+	}
+
 }
 
 FPluginReferenceNodeInfo::FPluginReferenceNodeInfo(const FPluginIdentifier& InIdentifier, bool bInReferencers)
@@ -103,6 +126,7 @@ int32 FPluginReferenceNodeInfo::ProvisionSize(const FPluginIdentifier& InParentI
 
 UEdGraph_PluginReferenceViewer::UEdGraph_PluginReferenceViewer(const FObjectInitializer& ObjectInitializer)
 	: Super(ObjectInitializer)
+	, bAdvancedInfoLoaded(false)
 {
 
 }
@@ -133,6 +157,11 @@ void UEdGraph_PluginReferenceViewer::GetPluginAssets(const FPluginIdentifier& Pl
 	UE::PluginReferenceViewer::Private::GetPluginAssets(ForPlugin, OutAssetIdentifiers);
 }
 
+const FPluginStats& UEdGraph_PluginReferenceViewer::GetPluginStats(const FPluginIdentifier& Plugin)
+{
+	return StatsMap.FindChecked(Plugin);
+}
+
 void UEdGraph_PluginReferenceViewer::CachePluginDependencies(const TArray<TSharedRef<IPlugin>>& Plugins)
 {
 	PluginMap.Empty();
@@ -239,6 +268,63 @@ UEdGraphNode_PluginReference* UEdGraph_PluginReferenceViewer::ConstructNodes(con
 	return RefilterGraph();
 }
 
+void UEdGraph_PluginReferenceViewer::LoadAdvancedPluginInfo()
+{
+	// don't reload everything if we have already filled the stats map
+	if (bAdvancedInfoLoaded)
+	{
+		return;
+	}
+
+	// fill the map before we task work out
+	for (const TPair<FPluginIdentifier, TUniquePtr<FPluginDependsNode> >& CachedPair : CachedDependsNodes)
+	{
+		StatsMap.Add(CachedPair.Key);
+	}
+
+	TArray<FPluginIdentifier> MapKeys;
+	ReferencerNodeInfos.GenerateKeyArray(MapKeys);
+
+	// GetSizeOfPlugin is an expensive call so using ParallelFors here
+
+	ParallelFor(MapKeys.Num(),
+		[this, &MapKeys](int32 Index)
+		{
+			FPluginStats& Stats = StatsMap.FindChecked(MapKeys[Index]);
+			Stats.Size = UE::PluginReferenceViewer::Private::GetSizeOfPlugin(PluginMap.FindChecked(MapKeys[Index]));
+		});
+
+	MapKeys.Empty();
+	DependencyNodeInfos.GenerateKeyArray(MapKeys);
+
+	ParallelFor(MapKeys.Num(),
+		[this, &MapKeys](int32 Index)
+		{
+			FPluginStats& Stats = StatsMap.FindChecked(MapKeys[Index]);
+			Stats.Size = UE::PluginReferenceViewer::Private::GetSizeOfPlugin(PluginMap.FindChecked(MapKeys[Index]));
+		});
+
+
+	for (TPair<FPluginIdentifier, TUniquePtr<FPluginDependsNode> >& CachedPair : CachedDependsNodes)
+	{
+		FPluginDependsNode* DependsNode = CachedPair.Value.Get();
+
+		FPluginStats& CurrentPluginStats = StatsMap.FindChecked(CachedPair.Key);
+
+		CurrentPluginStats.Dependencies = DependsNode->Dependencies.Num();
+		CurrentPluginStats.Referencers = DependsNode->Referencers.Num();
+
+		CurrentPluginStats.SizeWithDependencies = CurrentPluginStats.Size;
+
+		for (const FPluginDependsNode* Dependency : DependsNode->Dependencies)
+		{
+			CurrentPluginStats.SizeWithDependencies += StatsMap.FindChecked(Dependency->Identifier).Size;
+		}
+	}
+
+	bAdvancedInfoLoaded = true;
+}
+
 UEdGraphNode_PluginReference* UEdGraph_PluginReferenceViewer::RecursivelyCreateNodes(bool bInReferencers, const FPluginIdentifier& InPluginId, const FIntPoint& InNodeLoc, const FPluginIdentifier& InParentId, UEdGraphNode_PluginReference* InParentNode, TMap<FPluginIdentifier, FPluginReferenceNodeInfo>& InNodeInfos, int32 InCurrentDepth, int32 InMaxDepth, bool bIsRoot)
 {
 	check(InNodeInfos.Contains(InPluginId));
@@ -387,6 +473,16 @@ UEdGraphNode_PluginReference* UEdGraph_PluginReferenceViewer::RebuildGraph()
 	RemoveAllNodes();
 
 	UEdGraphNode_PluginReference* NewRootNode = ConstructNodes(CurrentGraphRootIdentifiers, CurrentGraphRootOrigin);
+
+	if (bAdvancedInfoLoaded)
+	{
+		// Load advanced plugin checks this bool so multiple load requests from the UI don't cause needless reloads, so set false
+		bAdvancedInfoLoaded = false;
+
+		StatsMap.Empty();
+		LoadAdvancedPluginInfo();
+	}
+
 	return NewRootNode;
 }
 
@@ -536,4 +632,5 @@ void UEdGraph_PluginReferenceViewer::GetPluginReferencers(const FPluginIdentifie
 			}
 		}
 	}
-}
\ No newline at end of file
+}
+
diff --git a/Engine/Plugins/Experimental/PluginReferenceViewer/Source/PluginReferenceViewer/Private/EdGraph_PluginReferenceViewer.h b/Engine/Plugins/Experimental/PluginReferenceViewer/Source/PluginReferenceViewer/Private/EdGraph_PluginReferenceViewer.h
--- a/Engine/Plugins/Experimental/PluginReferenceViewer/Source/PluginReferenceViewer/Private/EdGraph_PluginReferenceViewer.h
+++ b/Engine/Plugins/Experimental/PluginReferenceViewer/Source/PluginReferenceViewer/Private/EdGraph_PluginReferenceViewer.h
@@ -53,6 +53,26 @@ struct FPluginIdentifier
 	}
 };
 
+
+struct FAssetStats
+{
+	FAssetStats() : References(0)
+	{
+	}
+
+	uint32 References;
+};
+
+struct FPluginStats
+{
+	uint64 Size;
+	uint64 SizeWithDependencies;
+
+	uint32 Dependencies;
+	uint32 Referencers;
+};
+
+
 struct FPluginDependsNode
 {
 	explicit FPluginDependsNode(const FPluginIdentifier& InIdentifier)
@@ -112,6 +132,10 @@ public:
 	/** Finds all plugin assets */
 	void GetPluginAssets(const FPluginIdentifier& Plugin, TArray<FAssetIdentifier>& OutAssetIdentifiers);
 
+	const FPluginStats& GetPluginStats(const FPluginIdentifier& Plugin);
+
+	void LoadAdvancedPluginInfo();
+
 	UEdGraphNode_PluginReference* ConstructNodes(const TArray<FPluginIdentifier>& GraphRootIdentifiers, const FIntPoint& GraphRootOrigin);
 
 	UEdGraphNode_PluginReference* RecursivelyCreateNodes(bool bInReferencers, const FPluginIdentifier& InPluginId, const FIntPoint& InNodeLoc, const FPluginIdentifier& InParentId, UEdGraphNode_PluginReference* InParentNode, TMap<FPluginIdentifier, FPluginReferenceNodeInfo>& InNodeInfos, int32 InCurrentDepth, int32 InMaxDepth, bool bIsRoot = false);
@@ -165,9 +189,12 @@ private:
 
 	TMap<FPluginIdentifier, TUniquePtr<FPluginDependsNode>> CachedDependsNodes;
 	TMap<FPluginIdentifier, const TSharedRef<IPlugin>> PluginMap;
+	TMap<FPluginIdentifier, FPluginStats> StatsMap;
 
 	TMap<FPluginIdentifier, FPluginReferenceNodeInfo> ReferencerNodeInfos;
 	TMap<FPluginIdentifier, FPluginReferenceNodeInfo> DependencyNodeInfos;
 
 	friend SPluginReferenceViewer;
+
+	bool bAdvancedInfoLoaded = false;
 };
\ No newline at end of file
diff --git a/Engine/Plugins/Experimental/PluginReferenceViewer/Source/PluginReferenceViewer/Private/PluginReferenceViewerModule.cpp b/Engine/Plugins/Experimental/PluginReferenceViewer/Source/PluginReferenceViewer/Private/PluginReferenceViewerModule.cpp
--- a/Engine/Plugins/Experimental/PluginReferenceViewer/Source/PluginReferenceViewer/Private/PluginReferenceViewerModule.cpp
+++ b/Engine/Plugins/Experimental/PluginReferenceViewer/Source/PluginReferenceViewer/Private/PluginReferenceViewerModule.cpp
@@ -52,11 +52,13 @@ void FPluginReferenceViewerModule::ShutdownModule()
 
 TSharedRef<SDockTab> FPluginReferenceViewerModule::SpawnPluginReferenceViewerTab(const FSpawnTabArgs& Args)
 {
-	return SNew(SDockTab)
-		.TabRole(ETabRole::NomadTab)
-		[
-			SNew(SPluginReferenceViewer)
-		];
+	TSharedRef<SDockTab> WidgetReflectorTab = SNew(SDockTab)
+		.TabRole(ETabRole::NomadTab);
+
+	WidgetReflectorTab->SetContent(SNew(SPluginReferenceViewer)
+		.ParentTab(WidgetReflectorTab));
+
+	return WidgetReflectorTab;
 }
 
 void FPluginReferenceViewerModule::OnLaunchReferenceViewerFromPluginBrowser(TSharedPtr<IPlugin> Plugin)
diff --git a/Engine/Plugins/Experimental/PluginReferenceViewer/Source/PluginReferenceViewer/Private/SPluginReferenceViewer.cpp b/Engine/Plugins/Experimental/PluginReferenceViewer/Source/PluginReferenceViewer/Private/SPluginReferenceViewer.cpp
--- a/Engine/Plugins/Experimental/PluginReferenceViewer/Source/PluginReferenceViewer/Private/SPluginReferenceViewer.cpp
+++ b/Engine/Plugins/Experimental/PluginReferenceViewer/Source/PluginReferenceViewer/Private/SPluginReferenceViewer.cpp
@@ -7,17 +7,44 @@
 #include "Features/EditorFeatures.h"
 #include "Features/IModularFeatures.h"
 #include "Features/IPluginsEditorFeature.h"
+#include "Framework/Docking/TabManager.h"
 #include "Interfaces/IPluginManager.h"
+#include "Math/UnitConversion.h"
 #include "PluginReferenceViewerCommands.h"
 #include "PluginReferenceViewerSchema.h"
 #include "PluginReferenceViewerStyle.h"
 #include "ToolMenus.h"
+#include "Widgets/Docking/SDockTab.h"
+#include "Widgets/Input/SCheckBox.h"
 #include "Widgets/Input/SSpinBox.h"
 #include "Widgets/Input/SEditableTextBox.h"
 #include "Widgets/Layout/SBorder.h"
 
 #define LOCTEXT_NAMESPACE "PluginReferenceViewer"
 
+namespace PluginReferenceViewerTabID
+{
+	static const FName ReferenceGraph = "PluginReferenceViewer.ReferenceGraph";
+	static const FName PluginStats = "PluginReferenceViewer.PluginStats";
+}
+
+namespace
+{
+	/** Provided a byte count this function outputs best fit representation of size using FUnitConversion
+	* @param SizeInBytes The amount of bytes to convert to a formatted text that represents it.
+	* @return A text 
+	*/
+	FText GetSizeInBytesAsBiggerUnits(const uint64 SizeInBytes)
+	{
+		const FNumericUnit<uint64> Unit = FUnitConversion::QuantizeUnitsToBestFit(SizeInBytes, EUnit::Bytes);
+
+		FString UnitAsString = FString::Printf(TEXT("%lld %s"), Unit.Value, FUnitConversion::GetUnitDisplayString(Unit.Units));
+		
+		return FText::FromString(UnitAsString);
+	}
+}
+
+
 SPluginReferenceViewer::~SPluginReferenceViewer()
 {
 	if (!GExitPurge)
@@ -34,6 +61,9 @@
 	/** Used to delay graph rebuilding during spinbox slider interaction */
 	bNeedsGraphRebuild = false;
 
+	/** Enables loading of plugin stats, such as reference counts, plugin content size, and dependant content size. */
+	bShowAdvancedInfo = false;
+
 	RegisterActions();
 
 	// Set up the history manager
@@ -60,175 +90,58 @@
 
 	GraphObj->CachePluginDependencies(IPluginManager::Get().GetDiscoveredPlugins());
 
+
+	const FName TabLayoutName = "PluginReferenceViewer_Layout_NoStats";
+
+	TSharedRef<FTabManager::FLayout> Layout = FTabManager::NewLayout(TabLayoutName)
+	->AddArea
+	(
+		FTabManager::NewPrimaryArea()
+		->SetOrientation(Orient_Vertical)
+		->Split
+		(
+			// Main application area
+			FTabManager::NewSplitter()
+			->SetOrientation(Orient_Vertical)
+			->Split
+			(
+				FTabManager::NewStack()
+				->SetHideTabWell(true)
+				->SetSizeCoefficient(0.7f)
+				->AddTab(PluginReferenceViewerTabID::ReferenceGraph, ETabState::OpenedTab)
+			)
+			->Split
+			(
+				FTabManager::NewStack()
+				->SetHideTabWell(true)
+				->SetSizeCoefficient(0.3f)
+				->AddTab(PluginReferenceViewerTabID::PluginStats, ETabState::ClosedTab)
+			)
+		)
+	);
+
+	auto RegisterTrackedTabSpawner = [this](const FName& TabId, const FOnSpawnTab& OnSpawnTab) -> FTabSpawnerEntry&
+	{
+		return TabManager->RegisterTabSpawner(TabId, FOnSpawnTab::CreateLambda([this, OnSpawnTab](const FSpawnTabArgs& Args) -> TSharedRef<SDockTab>
+			{
+				TSharedRef<SDockTab> SpawnedTab = OnSpawnTab.Execute(Args);
+				OnTabSpawned(Args.GetTabId().TabType, SpawnedTab);
+				return SpawnedTab;
+			}));
+	};
+
+	check(InArgs._ParentTab.IsValid());
+	TabManager = FGlobalTabmanager::Get()->NewTabManager(InArgs._ParentTab.ToSharedRef());
+
+	RegisterTrackedTabSpawner(PluginReferenceViewerTabID::ReferenceGraph, FOnSpawnTab::CreateSP(this, &SPluginReferenceViewer::SpawnTab_ReferenceGraph))
+		.SetDisplayName(LOCTEXT("ReferenceGraphTab", "Reference Graph"));
+
+	RegisterTrackedTabSpawner(PluginReferenceViewerTabID::PluginStats, FOnSpawnTab::CreateSP(this, &SPluginReferenceViewer::SpawnTab_PluginStats))
+		.SetDisplayName(LOCTEXT("PluginStatsTab", "Plugin Stats"));
+
 	ChildSlot
 	[
-		SNew(SVerticalBox)
-
-		// Path and history
-		+SVerticalBox::Slot()
-		.AutoHeight()
-		[
-			SNew(SBorder)
-			.BorderImage(FAppStyle::GetBrush("Brushes.Panel"))
-			[
-				SNew(SHorizontalBox)
-
-				+SHorizontalBox::Slot()
-				.AutoWidth()
-				.Padding(4, 0)
-				[
-					MakeToolBar()
-				]
-
-				// Path
-				+SHorizontalBox::Slot()
-				.Padding(0, 0, 4, 0)
-				.FillWidth(1.f)
-				.VAlign(VAlign_Center)
-				[
-					SNew(SBorder)
-					.BorderImage( FAppStyle::GetBrush("ToolPanel.GroupBorder") )
-					[
-						SNew(SEditableTextBox)
-						.Text(this, &SPluginReferenceViewer::GetAddressBarText)
-						.OnTextCommitted(this, &SPluginReferenceViewer::OnAddressBarTextCommitted)
-						.OnTextChanged(this, &SPluginReferenceViewer::OnAddressBarTextChanged)
-						.SelectAllTextWhenFocused(true)
-						.SelectAllTextOnCommit(true)
-						.Style(FAppStyle::Get(), "ReferenceViewer.PathText")
-					]
-				]
-			]
-		]
-
-		// Graph
-		+SVerticalBox::Slot()
-		.FillHeight(0.90f)
-		.HAlign(HAlign_Fill)
-		[
-			SNew(SOverlay)
-
-			+ SOverlay::Slot()
-			[
-				GraphEditorPtr.ToSharedRef()
-			]
-
-			+SOverlay::Slot()
-			.VAlign(VAlign_Top)
-			.HAlign(HAlign_Fill)
-			.Padding(8)
-			[
-				SNew(SHorizontalBox)
-				+ SHorizontalBox::Slot()
-				.AutoWidth()
-				[
-					SNew(SBorder)
-					.BorderImage(FAppStyle::GetBrush("ToolPanel.GroupBorder"))
-					[
-						SNew(SVerticalBox)
-						+ SVerticalBox::Slot()
-						.AutoHeight()
-						[
-							SNew(SHorizontalBox)
-							+ SHorizontalBox::Slot()
-							.VAlign(VAlign_Center)
-							.Padding(2.f)
-							[
-								SNew(STextBlock)
-								.Text(LOCTEXT("SearchDepthReferencersLabelText", "Search Referencers Depth"))
-							]
-
-							+ SHorizontalBox::Slot()
-							.AutoWidth()
-							.VAlign(VAlign_Center)
-							.Padding(2.f)
-							[
-								SNew(SBox)
-								.WidthOverride(100)
-								[
-									SAssignNew(ReferencerCountBox, SSpinBox<int32>)
-									.Value(this, &SPluginReferenceViewer::GetSearchReferencerDepthCount)
-									.OnValueChanged_Lambda([this] (int32 NewValue)
-										{	
-											if (NewValue != Settings.MaxSearchReferencersDepth)
-											{
-												Settings.MaxSearchReferencersDepth = NewValue;
-												bNeedsGraphRebuild = true;
-											}
-										}
-									)
-									.OnValueCommitted_Lambda([this] (int32 NewValue, ETextCommit::Type CommitType) 
-										{ 
-											FSlateApplication::Get().SetKeyboardFocus(GraphEditorPtr, EFocusCause::SetDirectly); 
-
-											if (NewValue != Settings.MaxSearchReferencersDepth || bNeedsGraphRebuild)
-											{
-												Settings.MaxSearchReferencersDepth = NewValue;
-												bNeedsGraphRebuild = false;
-												RebuildGraph();
-											}
-										}
-									)
-									.MinValue(0)
-									.MaxValue(50)
-									.MaxSliderValue(10)
-								]
-							]
-						]
-
-						+ SVerticalBox::Slot()
-						.AutoHeight()
-						[
-							SNew(SHorizontalBox)
-							+ SHorizontalBox::Slot()
-							.VAlign(VAlign_Center)
-							.Padding(2.f)
-							[
-								SNew(STextBlock)
-								.Text(LOCTEXT("SearchDepthDependenciesLabelText", "Search Dependencies Depth"))
-							]
-
-							+ SHorizontalBox::Slot()
-							.AutoWidth()
-							.VAlign(VAlign_Center)
-							.Padding(2.f)
-							[
-								SNew(SBox)
-								.WidthOverride(100)
-								[
-									SAssignNew(DependencyCountBox, SSpinBox<int32>)
-									.Value(this, &SPluginReferenceViewer::GetSearchDependencyDepthCount)
-									.OnValueChanged_Lambda([this] (int32 NewValue)
-										{	
-											if (NewValue != Settings.MaxSearchDependencyDepth)
-											{
-												Settings.MaxSearchDependencyDepth = NewValue;
-												bNeedsGraphRebuild = true;
-											}
-										}
-									)
-									.OnValueCommitted_Lambda([this] (int32 NewValue, ETextCommit::Type CommitType) 
-										{ 
-											FSlateApplication::Get().SetKeyboardFocus(GraphEditorPtr, EFocusCause::SetDirectly); 
-
-											if (NewValue != Settings.MaxSearchDependencyDepth || bNeedsGraphRebuild)
-											{
-												Settings.MaxSearchDependencyDepth = NewValue;
-												bNeedsGraphRebuild = false;
-												RebuildGraph();
-											}
-										}
-									)
-									.MinValue(0)
-									.MaxValue(50)
-									.MaxSliderValue(10)
-								]
-							]
-						]
-					]
-				]
-			]
-		]
+		TabManager->RestoreFrom(Layout, nullptr).ToSharedRef()
 	];
 }
 
@@ -429,6 +342,23 @@
 	return FText::GetEmpty();
 }
 
+FText SPluginReferenceViewer::GetPluginStatsText() const
+{
+	const UEdGraphNode* SelectedNode = GraphEditorPtr->GetSingleSelectedNode();
+
+	if (SelectedNode == nullptr)
+	{
+		return FText(LOCTEXT("PluginSelectionError", "Invalid number of nodes selected. Please select a single node!"));
+	}
+
+	FText PluginName = SelectedNode->GetNodeTitle(ENodeTitleType::FullTitle);
+
+	const FPluginStats & PluginStats = GraphObj->GetPluginStats(FPluginIdentifier(FName(PluginName.ToString())));
+
+	return FText::Format(LOCTEXT("PluginInfoText", "Name: {0}\nContent Size: {1}\nTotal Content Size With Dependencies: {2}\nDependencies: {3}\nReferencers: {4}"), 
+		{ PluginName, GetSizeInBytesAsBiggerUnits(PluginStats.Size), GetSizeInBytesAsBiggerUnits(PluginStats.SizeWithDependencies), PluginStats.Dependencies, PluginStats.Referencers });
+}
+
 FText SPluginReferenceViewer::GetAddressBarText() const
 {
 	if (GraphObj)
@@ -635,4 +565,258 @@
 	return Settings.MaxSearchDependencyDepth;
 }
 
+TSharedRef<SDockTab> SPluginReferenceViewer::SpawnTab_ReferenceGraph(const FSpawnTabArgs& Args)
+{
+	TSharedRef<SDockTab> SpawnedTab = SNew(SDockTab)
+	[
+		SNew(SVerticalBox)
+
+		// Path and history
+		+SVerticalBox::Slot()
+		.AutoHeight()
+		[
+			SNew(SBorder)
+			.BorderImage(FAppStyle::GetBrush("Brushes.Panel"))
+			[
+				SNew(SHorizontalBox)
+
+				+SHorizontalBox::Slot()
+				.AutoWidth()
+				.Padding(4, 0)
+				[
+					MakeToolBar()
+				]
+
+				// Path
+				+SHorizontalBox::Slot()
+				.Padding(0, 0, 4, 0)
+				.FillWidth(1.f)
+				.VAlign(VAlign_Center)
+				[
+					SNew(SBorder)
+					.BorderImage( FAppStyle::GetBrush("ToolPanel.GroupBorder") )
+					[
+						SNew(SEditableTextBox)
+						.Text(this, &SPluginReferenceViewer::GetAddressBarText)
+						.OnTextCommitted(this, &SPluginReferenceViewer::OnAddressBarTextCommitted)
+						.OnTextChanged(this, &SPluginReferenceViewer::OnAddressBarTextChanged)
+						.SelectAllTextWhenFocused(true)
+						.SelectAllTextOnCommit(true)
+						.Style(FAppStyle::Get(), "ReferenceViewer.PathText")
+					]
+				]
+			]
+		]
+
+		// Graph
+		+SVerticalBox::Slot()
+		.FillHeight(0.90f)
+		.HAlign(HAlign_Fill)
+		[
+			SNew(SOverlay)
+
+			+ SOverlay::Slot()
+			[
+				GraphEditorPtr.ToSharedRef()
+			]
+
+			+SOverlay::Slot()
+			.VAlign(VAlign_Top)
+			.HAlign(HAlign_Fill)
+			.Padding(8)
+			[
+				SNew(SHorizontalBox)
+				+ SHorizontalBox::Slot()
+				.AutoWidth()
+				[
+					SNew(SBorder)
+					.BorderImage(FAppStyle::GetBrush("ToolPanel.GroupBorder"))
+					[
+						SNew(SVerticalBox)
+						+ SVerticalBox::Slot()
+						.AutoHeight()
+						[
+							SNew(SHorizontalBox)
+							+ SHorizontalBox::Slot()
+							.VAlign(VAlign_Center)
+							.Padding(2.f)
+							[
+								SNew(STextBlock)
+								.Text(LOCTEXT("SearchDepthReferencersLabelText", "Search Referencers Depth"))
+							]
+
+							+ SHorizontalBox::Slot()
+							.AutoWidth()
+							.VAlign(VAlign_Center)
+							.Padding(2.f)
+							[
+								SNew(SBox)
+								.WidthOverride(100)
+								[
+									SAssignNew(ReferencerCountBox, SSpinBox<int32>)
+									.Value(this, &SPluginReferenceViewer::GetSearchReferencerDepthCount)
+									.OnValueChanged_Lambda([this] (int32 NewValue)
+										{	
+											if (NewValue != Settings.MaxSearchReferencersDepth)
+											{
+												Settings.MaxSearchReferencersDepth = NewValue;
+												bNeedsGraphRebuild = true;
+											}
+										}
+									)
+									.OnValueCommitted_Lambda([this] (int32 NewValue, ETextCommit::Type CommitType) 
+										{ 
+											FSlateApplication::Get().SetKeyboardFocus(GraphEditorPtr, EFocusCause::SetDirectly); 
+
+											if (NewValue != Settings.MaxSearchReferencersDepth || bNeedsGraphRebuild)
+											{
+												Settings.MaxSearchReferencersDepth = NewValue;
+												bNeedsGraphRebuild = false;
+												RebuildGraph();
+											}
+										}
+									)
+									.MinValue(0)
+									.MaxValue(50)
+									.MaxSliderValue(10)
+								]
+							]
+						]
+
+						+ SVerticalBox::Slot()
+						.AutoHeight()
+						[
+							SNew(SHorizontalBox)
+							+ SHorizontalBox::Slot()
+							.VAlign(VAlign_Center)
+							.Padding(2.f)
+							[
+								SNew(STextBlock)
+								.Text(LOCTEXT("SearchDepthDependenciesLabelText", "Search Dependencies Depth"))
+							]
+
+							+ SHorizontalBox::Slot()
+							.AutoWidth()
+							.VAlign(VAlign_Center)
+							.Padding(2.f)
+							[
+								SNew(SBox)
+								.WidthOverride(100)
+								[
+									SAssignNew(DependencyCountBox, SSpinBox<int32>)
+									.Value(this, &SPluginReferenceViewer::GetSearchDependencyDepthCount)
+									.OnValueChanged_Lambda([this] (int32 NewValue)
+										{	
+											if (NewValue != Settings.MaxSearchDependencyDepth)
+											{
+												Settings.MaxSearchDependencyDepth = NewValue;
+												bNeedsGraphRebuild = true;
+											}
+										}
+									)
+									.OnValueCommitted_Lambda([this] (int32 NewValue, ETextCommit::Type CommitType) 
+										{ 
+											FSlateApplication::Get().SetKeyboardFocus(GraphEditorPtr, EFocusCause::SetDirectly); 
+
+											if (NewValue != Settings.MaxSearchDependencyDepth || bNeedsGraphRebuild)
+											{
+												Settings.MaxSearchDependencyDepth = NewValue;
+												bNeedsGraphRebuild = false;
+												RebuildGraph();
+											}
+										}
+									)
+									.MinValue(0)
+									.MaxValue(50)
+									.MaxSliderValue(10)
+								]
+							]
+						]
+						+ SVerticalBox::Slot()
+						.AutoHeight()
+						[
+							SNew(SHorizontalBox)
+							+ SHorizontalBox::Slot()
+							.VAlign(VAlign_Center)
+							.Padding(2.f)
+							[
+								SNew(STextBlock)
+								.Text(LOCTEXT("ShowAdvancedInfoLabelText", "Show Advanced Info"))
+							]
+
+							+ SHorizontalBox::Slot()
+							.AutoWidth()
+							.VAlign(VAlign_Center)
+							.Padding(2.f)
+							[
+								SNew(SCheckBox)
+									.IsChecked_Lambda([this]()
+										{
+											return bShowAdvancedInfo ? ECheckBoxState::Checked : ECheckBoxState::Unchecked;
+										})
+									.OnCheckStateChanged_Lambda([this](ECheckBoxState InCheckBoxState)
+										{
+											if (InCheckBoxState == ECheckBoxState::Checked && GraphObj != nullptr)
+											{
+												// need to find a way to do a loading thing here
+												GraphObj->LoadAdvancedPluginInfo();
+												bShowAdvancedInfo = true;
+												TabManager->TryInvokeTab(PluginReferenceViewerTabID::PluginStats);
+											}
+											else
+											{
+												TSharedPtr<SDockTab> TabPtr = TabManager->FindExistingLiveTab(PluginReferenceViewerTabID::PluginStats);
+												if (TabPtr.IsValid())
+												{
+													TabPtr->RequestCloseTab();
+												}
+												// we don't want to actually change our state if we weren't able to load advanced info
+												bShowAdvancedInfo = false;
+											}
+										})
+							]
+						]
+					]
+
+				]
+			]
+		]
+	];
+
+	return SpawnedTab;
+}
+
+TSharedRef<SDockTab> SPluginReferenceViewer::SpawnTab_PluginStats(const FSpawnTabArgs& Args) const
+{
+	const TSharedRef<SDockTab> DockTab = SNew(SDockTab)
+		.ShouldAutosize(false)
+		.TabRole(ETabRole::PanelTab)
+		[
+			SNew(SVerticalBox)
+			+ SVerticalBox::Slot()
+			[
+				SNew(STextBlock)
+				.Text(this, &SPluginReferenceViewer::GetPluginStatsText)
+			]
+		];
+
+	return DockTab;
+}
+
+void SPluginReferenceViewer::OnTabSpawned(const FName& TabIdentifier, const TSharedRef<SDockTab>& SpawnedTab)
+{
+	TWeakPtr<SDockTab>* const ExistingTab = SpawnedTabs.Find(TabIdentifier);
+	if (ExistingTab)
+	{
+		if (ExistingTab->IsValid())
+		{
+			*ExistingTab = SpawnedTab;
+		}
+	}
+	else
+	{
+		SpawnedTabs.Add(TabIdentifier, SpawnedTab);
+	}
+}
+
 #undef LOCTEXT_NAMESPACE
diff --git a/Engine/Plugins/Experimental/PluginReferenceViewer/Source/PluginReferenceViewer/Private/SPluginReferenceViewer.h b/Engine/Plugins/Experimental/PluginReferenceViewer/Source/PluginReferenceViewer/Private/SPluginReferenceViewer.h
--- a/Engine/Plugins/Experimental/PluginReferenceViewer/Source/PluginReferenceViewer/Private/SPluginReferenceViewer.h
+++ b/Engine/Plugins/Experimental/PluginReferenceViewer/Source/PluginReferenceViewer/Private/SPluginReferenceViewer.h
@@ -7,12 +7,15 @@
 #include "GraphEditor.h"
 #include "HistoryManager.h"
 
+#include "Framework/Docking/TabManager.h"
+
 //////////////////////////////////////////////////////////////////////////
 // SPluginReferenceViewer
 
 class IPlugin;
 class FUICommandInfo;
 class FUICommandList;
+class SDockTab;
 class UEdGraphNode;
 class UEdGraph_PluginReferenceViewer;
 
@@ -22,6 +25,8 @@ public:
 	SLATE_BEGIN_ARGS(SPluginReferenceViewer)
 	{}
 
+	SLATE_ARGUMENT(TSharedPtr<SDockTab>, ParentTab)
+
 	SLATE_END_ARGS()
 
 	~SPluginReferenceViewer();
@@ -41,6 +46,7 @@ public:
 private:
 	void OnOpenPluginProperties();
 	bool HasAtLeastOneRealNodeSelected();
+
 	void OpenPluginProperties(const FString& PluginName);
 
 	TSharedRef<SWidget> MakeToolBar();
@@ -82,6 +88,9 @@ private:
 	/** Gets the tool tip text for the history forward button */
 	FText GetHistoryForwardTooltip() const;
 
+	/** Gets text for advanced info plugin stats. Formatted output of plugin content size, reference counts, etc. */
+	FText GetPluginStatsText() const;
+
 	FText GetAddressBarText() const;
 	void OnAddressBarTextCommitted(const FText& NewText, ETextCommit::Type CommitInfo);
 	void OnAddressBarTextChanged(const FText& NewText);
@@ -103,6 +112,10 @@ private:
 		bool bShowOptionalPlugins = true;
 	};
 
+	void OnTabSpawned(const FName& TabIdentifier, const TSharedRef<SDockTab>& SpawnedTab);
+	TSharedRef<SDockTab> SpawnTab_ReferenceGraph(const FSpawnTabArgs& Args);
+	TSharedRef<SDockTab> SpawnTab_PluginStats(const FSpawnTabArgs& Args) const;
+
 private:
 	/** The manager that keeps track of history data for this browser */
 	FPluginReferenceViewerHistoryManager HistoryManager;
@@ -118,8 +131,14 @@ private:
 	/** The temporary copy of the path text when it is actively being edited. */
 	FText TemporaryPathBeingEdited;
 
+	TSharedPtr<FTabManager> TabManager;
+	TMap<FName, TWeakPtr<SDockTab>> SpawnedTabs;
+
 	TSharedPtr<FUICommandInfo> ShowEnginePlugins;
 
 	/** Used to delay graph rebuilding during spinbox slider interaction */
 	bool bNeedsGraphRebuild;
+
+	/** Enables loading of plugin stats, such as reference counts, plugin content size, and dependant content size. */
+	bool bShowAdvancedInfo;
 };
\ No newline at end of file
