Plugin Reference Viewer - Proper handling of plugin dependencies to  plugins that have not been discovered

diff --git a/Engine/Plugins/Experimental/PluginReferenceViewer/Source/PluginReferenceViewer/Private/EdGraphNode_PluginReference.cpp b/Engine/Plugins/Experimental/PluginReferenceViewer/Source/PluginReferenceViewer/Private/EdGraphNode_PluginReference.cpp
--- a/Engine/Plugins/Experimental/PluginReferenceViewer/Source/PluginReferenceViewer/Private/EdGraphNode_PluginReference.cpp
+++ b/Engine/Plugins/Experimental/PluginReferenceViewer/Source/PluginReferenceViewer/Private/EdGraphNode_PluginReference.cpp
@@ -36,10 +36,18 @@ void UEdGraphNode_PluginReference::SetupPluginReferenceNode(const FIntPoint& InN
 	AssetBrush = FSlateIcon(FName("PluginStyle"), "Plugins.TabIcon");
 
 	PluginIdentifier = InPluginIdentifier;
-	NodeTitle = FText::FromString(InPlugin->GetName());
+
+	if (InPlugin.IsValid())
+	{
+		NodeTitle = FText::FromString(InPlugin->GetName());
+		bIsEnginePlugin = (InPlugin->GetType() == EPluginType::Engine) ? true : false;
+	}
+	else
+	{
+		NodeTitle = FText::Format(LOCTEXT("Title_PluginNotFound", "{0} **Plugin Not Found**"), FText::FromName(InPluginIdentifier.Name));
+	}
 
 	bAllowThumbnail = bInAllowThumbnail;
-	bIsEnginePlugin = (InPlugin->GetType() == EPluginType::Engine) ? true : false;
 	bIsADuplicate = bInIsADuplicate;
 
 	CachedPlugin = InPlugin;
@@ -99,22 +107,35 @@ FText UEdGraphNode_PluginReference::GetNodeTitle(ENodeTitleType::Type TitleType)
 
 FLinearColor UEdGraphNode_PluginReference::GetNodeTitleColor() const
 {
-	if (bIsEnginePlugin)
+	if (CachedPlugin.IsValid())
 	{
-		return FLinearColor(0.55f, 0.55f, 0.55f);
+		if (bIsEnginePlugin)
+		{
+			return FLinearColor(0.55f, 0.55f, 0.55f);
+		}
+		else
+		{
+			return FLinearColor(0.0f, 0.42f, 1.0f);
+		}
 	}
 	else
 	{
-		return FLinearColor(0.0f, 0.42f, 1.0f);
+		return FLinearColor(0.55f, 0.55f, 0.55f);
 	}
 }
 
 FText UEdGraphNode_PluginReference::GetTooltipText() const
 {
-	FString TooltipText = CachedPlugin->GetBaseDir();
-	FPaths::MakePathRelativeTo(TooltipText, *FPaths::RootDir());
-	
-	return FText::FromString(TooltipText);
+	if (CachedPlugin.IsValid())
+	{
+		FString TooltipText = CachedPlugin->GetBaseDir();
+		FPaths::MakePathRelativeTo(TooltipText, *FPaths::RootDir());
+		return FText::FromString(TooltipText);
+	}
+	else
+	{
+		return LOCTEXT("Tooltip_PluginNotFound", "Plugin not found!");
+	}
 }
 
 UEdGraphPin* UEdGraphNode_PluginReference::GetDependencyPin()
diff --git a/Engine/Plugins/Experimental/PluginReferenceViewer/Source/PluginReferenceViewer/Private/EdGraph_PluginReferenceViewer.cpp b/Engine/Plugins/Experimental/PluginReferenceViewer/Source/PluginReferenceViewer/Private/EdGraph_PluginReferenceViewer.cpp
--- a/Engine/Plugins/Experimental/PluginReferenceViewer/Source/PluginReferenceViewer/Private/EdGraph_PluginReferenceViewer.cpp
+++ b/Engine/Plugins/Experimental/PluginReferenceViewer/Source/PluginReferenceViewer/Private/EdGraph_PluginReferenceViewer.cpp
@@ -253,7 +253,11 @@ UEdGraphNode_PluginReference* UEdGraph_PluginReferenceViewer::ConstructNodes(con
 
 		for (TPair<FPluginIdentifier, FPluginReferenceNodeInfo>& InfoPair : NewDependencyNodeInfos)
 		{
-			InfoPair.Value.Plugin = PluginMap.FindChecked(InfoPair.Key);
+			const TSharedRef<IPlugin>* FoundPlugin = PluginMap.Find(InfoPair.Key);
+			if (FoundPlugin != nullptr)
+			{
+				InfoPair.Value.Plugin = *FoundPlugin;
+			}
 		}
 
 		ReferencerNodeInfos = NewReferenceNodeInfos;
@@ -594,10 +598,13 @@ void UEdGraph_PluginReferenceViewer::GetPluginDependencies(const FPluginIdentifi
 	{
 		const FPluginDependsNode* ChildNode = PluginNode->Dependencies[i];
 
-		const TSharedRef<IPlugin> ChildPlugin = PluginMap.FindChecked(ChildNode->Identifier);
-		if (ChildPlugin->GetType() == EPluginType::Engine && !bShowEnginePlugins)
+		// Plugin dependencies may not be discovered
+		if (const TSharedRef<IPlugin>* ChildPlugin = PluginMap.Find(ChildNode->Identifier))
 		{
-			continue;
+			if ((*ChildPlugin)->GetType() == EPluginType::Engine && !bShowEnginePlugins)
+			{
+				continue;
+			}
 		}
 
 		const FPluginReferenceDescriptor& ReferenceDescriptor = Plugin->GetDescriptor().Plugins[i];
diff --git a/Engine/Plugins/Experimental/PluginReferenceViewer/Source/PluginReferenceViewer/Private/PluginReferenceViewerSchema.cpp b/Engine/Plugins/Experimental/PluginReferenceViewer/Source/PluginReferenceViewer/Private/PluginReferenceViewerSchema.cpp
--- a/Engine/Plugins/Experimental/PluginReferenceViewer/Source/PluginReferenceViewer/Private/PluginReferenceViewerSchema.cpp
+++ b/Engine/Plugins/Experimental/PluginReferenceViewer/Source/PluginReferenceViewer/Private/PluginReferenceViewerSchema.cpp
@@ -49,6 +49,12 @@ UPluginReferenceViewerSchema::UPluginReferenceViewerSchema(const FObjectInitiali
 
 void UPluginReferenceViewerSchema::GetContextMenuActions(UToolMenu* Menu, UGraphNodeContextMenuContext* Context) const
 {
+	const UEdGraphNode_PluginReference* ContextPluginNode = Cast<UEdGraphNode_PluginReference>(Context->Node);
+	if (!ContextPluginNode->GetPlugin().IsValid())
+	{
+		return;
+	}
+
 	{
 		FToolMenuSection& Section = Menu->AddSection(TEXT("Plugin"), NSLOCTEXT("PluginReferenceViewerSchema", "PluginSectionLabel", "Plugin"));
 		Section.AddMenuEntry(FPluginReferenceViewerCommands::Get().OpenPluginProperties);
@@ -60,7 +66,6 @@ void UPluginReferenceViewerSchema::GetContextMenuActions(UToolMenu* Menu, UGraph
 		FText Title = NSLOCTEXT("PluginReferenceViewerSchema", "OpenReferenceViewer", "Open {0} Asset References...");
 
 		UEdGraph_PluginReferenceViewer* PluginReferenceViewer = const_cast<UEdGraph_PluginReferenceViewer*>(CastChecked<UEdGraph_PluginReferenceViewer>(Context->Graph));
-		const UEdGraphNode_PluginReference* ContextPluginNode = Cast<UEdGraphNode_PluginReference>(Context->Node);
 
 		// Opens the reference viewer for all assets belonging to the plugin.
 		check(!PluginReferenceViewer->GetCurrentGraphRootIdentifiers().IsEmpty());
diff --git a/Engine/Plugins/Experimental/PluginReferenceViewer/Source/PluginReferenceViewer/Private/SPluginReferenceNode.cpp b/Engine/Plugins/Experimental/PluginReferenceViewer/Source/PluginReferenceViewer/Private/SPluginReferenceNode.cpp
--- a/Engine/Plugins/Experimental/PluginReferenceViewer/Source/PluginReferenceViewer/Private/SPluginReferenceNode.cpp
+++ b/Engine/Plugins/Experimental/PluginReferenceViewer/Source/PluginReferenceViewer/Private/SPluginReferenceNode.cpp
@@ -23,10 +23,12 @@ void SPluginReferenceNode::Construct(const FArguments& InArgs, UEdGraphNode_Plug
 	const int32 ThumbnailSize = 128;
 	if (InNode->AllowsThumbnail())
 	{
-		const TSharedPtr<const IPlugin> Plugin = InNode->GetPlugin();
-
 		// Plugin thumbnail image
-		FString Icon128FilePath = Plugin->GetBaseDir() / TEXT("Resources/Icon128.png");
+		FString Icon128FilePath;
+		if (InNode->GetPlugin().IsValid())
+		{
+			Icon128FilePath = InNode->GetPlugin()->GetBaseDir() / TEXT("Resources/Icon128.png");
+		}
 		if (!FPlatformFileManager::Get().GetPlatformFile().FileExists(*Icon128FilePath))
 		{
 			Icon128FilePath = IPluginManager::Get().FindPlugin(TEXT("PluginBrowser"))->GetBaseDir() / TEXT("Resources/DefaultIcon128.png");
diff --git a/Engine/Plugins/Experimental/PluginReferenceViewer/Source/PluginReferenceViewer/Private/SPluginReferenceViewer.cpp b/Engine/Plugins/Experimental/PluginReferenceViewer/Source/PluginReferenceViewer/Private/SPluginReferenceViewer.cpp
--- a/Engine/Plugins/Experimental/PluginReferenceViewer/Source/PluginReferenceViewer/Private/SPluginReferenceViewer.cpp
+++ b/Engine/Plugins/Experimental/PluginReferenceViewer/Source/PluginReferenceViewer/Private/SPluginReferenceViewer.cpp
@@ -460,9 +460,12 @@ void SPluginReferenceViewer::OnNodeDoubleClicked(UEdGraphNode* Node)
 {
 	if (UEdGraphNode_PluginReference* PlugingReferenceNode = Cast<UEdGraphNode_PluginReference>(Node))
 	{
-		TSet<UObject*> Nodes;
-		Nodes.Add(Node);
-		ReCenterGraphOnNodes(Nodes);
+		if (PlugingReferenceNode->GetPlugin().IsValid())
+		{
+			TSet<UObject*> Nodes;
+			Nodes.Add(Node);
+			ReCenterGraphOnNodes(Nodes);
+		}
 	}
 }
 
