Plugin reference viewer utility suite. Plugin reference graphs can be exported to .csv to detail the number of reference types to each of the plugin dependencies

diff --git a/Engine/Plugins/Experimental/PluginReferenceViewer/Source/PluginReferenceViewer/PluginReferenceViewer.Build.cs b/Engine/Plugins/Experimental/PluginReferenceViewer/Source/PluginReferenceViewer/PluginReferenceViewer.Build.cs
--- a/Engine/Plugins/Experimental/PluginReferenceViewer/Source/PluginReferenceViewer/PluginReferenceViewer.Build.cs
+++ b/Engine/Plugins/Experimental/PluginReferenceViewer/Source/PluginReferenceViewer/PluginReferenceViewer.Build.cs
@@ -1,6 +1,7 @@
 // Copyright Epic Games, Inc. All Rights Reserved.
 
 using UnrealBuildTool;
+using UnrealBuildTool.Rules;
 
 public class PluginReferenceViewer : ModuleRules
 {
@@ -31,6 +32,7 @@ public class PluginReferenceViewer : ModuleRules
 				"Projects",
 				"ApplicationCore",
 				"UnrealEd",
+				"GameplayTags",
 				"GraphEditor",
 				"EditorWidgets",
 				"ToolMenus",
diff --git a/Engine/Plugins/Experimental/PluginReferenceViewer/Source/PluginReferenceViewer/Private/PluginReferenceViewerUtils.cpp b/Engine/Plugins/Experimental/PluginReferenceViewer/Source/PluginReferenceViewer/Private/PluginReferenceViewerUtils.cpp
new file mode 100644
--- /dev/null
+++ b/Engine/Plugins/Experimental/PluginReferenceViewer/Source/PluginReferenceViewer/Private/PluginReferenceViewerUtils.cpp
@@ -0,0 +1,357 @@
+// Copyright Epic Games, Inc. All Rights Reserved.
+
+#include "PluginReferenceViewerUtils.h"
+
+#include "AssetRegistry/AssetRegistryModule.h"
+#include "AssetRegistry/AssetData.h"
+#include "AssetManagerEditorModule.h"
+#include "AssetRegistry/AssetData.h"
+#include "DesktopPlatformModule.h"
+#include "GameplayTagsManager.h"
+#include "GameplayTagsSettings.h"
+#include "HAL/IConsoleManager.h"
+#include "HAL/PlatformFileManager.h"
+#include "Interfaces/IPluginManager.h"
+#include "Misc/AssetRegistryInterface.h"
+#include "Misc/Paths.h"
+#include "Misc/ScopedSlowTask.h"
+
+#define LOCTEXT_NAMESPACE "PluginReferenceViewerUtils"
+
+DEFINE_LOG_CATEGORY_STATIC(LogPluginReferenceViewerUtils, Log, All);
+
+namespace PluginReferenceViewerUtils
+{
+	void ExportPluginReferenceGraph(const TArray<FString>& InArgs)
+	{
+		TArray<FString> PluginNames;
+		if (InArgs.Num() >= 1)
+		{
+			InArgs[0].ParseIntoArray(PluginNames, TEXT(","));
+		}
+		else
+		{
+			UE_LOG(LogPluginReferenceViewerUtils, Display, TEXT("Invalid plugin names argument. Expect plugin names seperated by ',' as the first argument. e.g a,b,c"));
+			return;
+		}
+
+		FString Filename;
+		if (InArgs.Num() >= 2)
+		{
+			Filename = InArgs[1];
+		}
+		else if (!PluginNames.IsEmpty())
+		{
+			Filename = FPaths::SetExtension(PluginNames[0], TEXT("csv"));
+		}
+
+		FPluginReferenceViewerUtils::ExportPlugins(PluginNames, Filename);
+	}
+
+	TArray<TSharedRef<const IPlugin>> GetGameplayTagSourcePlugins(FName TagName)
+	{
+		TArray<TSharedRef<const IPlugin>> Result;
+		UGameplayTagsManager& Manager = UGameplayTagsManager::Get();
+
+		FString Comment;
+		TArray<FName> TagSources;
+		bool bIsTagExplicit, bIsRestrictedTag, bAllowNonRestrictedChildren;
+		if (Manager.GetTagEditorData(TagName, Comment, TagSources, bIsTagExplicit, bIsRestrictedTag, bAllowNonRestrictedChildren))
+		{
+			for (const FName& TagSourceName : TagSources)
+			{
+				if (const FGameplayTagSource* TagSource = Manager.FindTagSource(TagSourceName))
+				{
+					switch (TagSource->SourceType)
+					{
+					case EGameplayTagSourceType::TagList:
+					{
+						const FString ContentFilePath = FPaths::GetPath(TagSource->SourceTagList->ConfigFileName) / TEXT("../../Content/");
+						FString RootContentPath;
+						if (FPackageName::TryConvertFilenameToLongPackageName(ContentFilePath, RootContentPath))
+						{
+							if (const TSharedPtr<IPlugin> FoundPlugin = IPluginManager::Get().FindPluginFromPath(RootContentPath))
+							{
+								Result.Add(FoundPlugin.ToSharedRef());
+							}
+						}
+						break;
+					}
+					case EGameplayTagSourceType::DataTable:
+					{
+						if (const TSharedPtr<IPlugin>& FoundPlugin = IPluginManager::Get().FindPluginFromPath(TagSource->SourceName.ToString()))
+						{
+							Result.Add(FoundPlugin.ToSharedRef());
+						}
+						break;
+					}
+					case EGameplayTagSourceType::Native:
+					{
+						if (const TSharedPtr<IPlugin>& FoundPlugin = IPluginManager::Get().GetModuleOwnerPlugin(TagSource->SourceName))
+						{
+							Result.Add(FoundPlugin.ToSharedRef());
+						}
+						break;
+					}
+					default:
+						break;
+					}
+				}
+			}
+		}
+		return Result;
+	}
+
+	TArray<FAssetIdentifier> GetAssetDependencies(const TSharedRef<IPlugin>& InPlugin)
+	{
+		IAssetRegistry& AssetRegistry = FModuleManager::LoadModuleChecked<FAssetRegistryModule>("AssetRegistry").Get();
+
+		FARFilter Filter;
+		Filter.bRecursivePaths = true;
+		Filter.PackagePaths.Add(FName(InPlugin->GetMountedAssetPath()));
+
+		TArray<FAssetData> AssetsInPlugin;
+		AssetRegistry.GetAssets(Filter, AssetsInPlugin);
+
+		TSet<FAssetIdentifier> UniqueDependencies;
+
+		for (const FAssetData& AssetData : AssetsInPlugin)
+		{
+			TArray<FAssetIdentifier> AssetReferencers;
+			AssetRegistry.GetDependencies(AssetData.PackageName, AssetReferencers);
+
+			for (const FAssetIdentifier& AssetReferencer : AssetReferencers)
+			{
+				UniqueDependencies.Add(AssetReferencer);
+			}
+		}
+
+		return UniqueDependencies.Array();
+	}
+
+	TMap<FString, TArray<FAssetIdentifier>> SplitByPlugins(const TSharedRef<IPlugin>& InOwningPlugin, const TArray<FAssetIdentifier>& InPluginDependencies)
+	{
+		TMap<FString, TArray<FAssetIdentifier>> Results;
+
+		const FName GameplayTagStructPackage = FGameplayTag::StaticStruct()->GetOutermost()->GetFName();
+		const FName NAME_GameplayTag = FGameplayTag::StaticStruct()->GetFName();
+
+		for (const FAssetIdentifier& AssetIdentifier : InPluginDependencies)
+		{
+			const FString PackageNameString = AssetIdentifier.PackageName.ToString();
+
+			if (AssetIdentifier.ObjectName == NAME_GameplayTag && AssetIdentifier.PackageName == GameplayTagStructPackage)
+			{
+				const TArray<TSharedRef<const IPlugin>> SourcePlugins = PluginReferenceViewerUtils::GetGameplayTagSourcePlugins(AssetIdentifier.ValueName);
+				for (const TSharedRef<const IPlugin>& SourcePlugin : SourcePlugins)
+				{
+					if (SourcePlugin != InOwningPlugin)
+					{
+						Results.FindOrAdd(SourcePlugin->GetName()).Add(AssetIdentifier);
+					}
+				}
+			}
+			else if (FStringView ModuleName; FPackageName::TryConvertScriptPackageNameToModuleName(PackageNameString, ModuleName))
+			{
+				const TSharedPtr<IPlugin> ModulePlugin = IPluginManager::Get().GetModuleOwnerPlugin(FName(ModuleName));
+				if (ModulePlugin.IsValid() && ModulePlugin != InOwningPlugin)
+				{
+					Results.FindOrAdd(ModulePlugin->GetName()).Add(AssetIdentifier);
+				}
+			}
+			else
+			{
+				const TSharedPtr<IPlugin> PackagePlugin = IPluginManager::Get().FindPluginFromPath(PackageNameString);
+				if (PackagePlugin.IsValid() && PackagePlugin != InOwningPlugin)
+				{
+					Results.FindOrAdd(PackagePlugin->GetName()).Add(AssetIdentifier);
+				}
+			}
+		}
+
+		return Results;
+	}
+
+	void SplitByReferenceType(const TArray<FAssetIdentifier>& InAssetIdentifiers, TArray<FAssetIdentifier>& OutAssetReferences, TArray<FAssetIdentifier>& OutScriptReferences, TArray<FAssetIdentifier>& OutNameReferences)
+	{
+		const FName GameplayTagStructPackage = FGameplayTag::StaticStruct()->GetOutermost()->GetFName();
+		const FName NAME_GameplayTag = FGameplayTag::StaticStruct()->GetFName();
+		const FName NAME_DataTable(TEXT("DataTable"));
+
+		TArray<FName> AllPackageNames;
+		for (const FAssetIdentifier& AssetIdentifier : InAssetIdentifiers)
+		{
+			AllPackageNames.Add(AssetIdentifier.PackageName);
+		}
+
+		TMap<FName, FAssetData> PackagesToAssetDataMap;
+		UE::AssetRegistry::GetAssetForPackages(AllPackageNames, PackagesToAssetDataMap);
+
+		for (const FAssetIdentifier& AssetIdentifier : InAssetIdentifiers)
+		{
+			const FString PackageNameString = AssetIdentifier.PackageName.ToString();
+
+			if (AssetIdentifier.ObjectName == NAME_GameplayTag && AssetIdentifier.PackageName == GameplayTagStructPackage)
+			{
+				OutNameReferences.Add(AssetIdentifier);
+			}
+			else if (FPackageName::IsScriptPackage(PackageNameString))
+			{
+				FStringView ModuleName;
+				if (FPackageName::TryConvertScriptPackageNameToModuleName(PackageNameString, ModuleName))
+				{
+					OutScriptReferences.Add(AssetIdentifier);
+				}
+			}
+			else
+			{
+				// PackagesToAssetDataMap may not contain the package as references to assets that no longer exist do occur.
+				const FAssetData* AssetData = PackagesToAssetDataMap.Find(AssetIdentifier.PackageName);
+				if (AssetData != nullptr)
+				{
+					if (AssetData->AssetClassPath.GetAssetName() == NAME_DataTable)
+					{
+						OutNameReferences.Add(AssetIdentifier);
+					}
+					else
+					{
+						OutAssetReferences.Add(AssetIdentifier);
+					}
+				}
+				else
+				{
+					UE_LOG(LogPluginReferenceViewerUtils, Display, TEXT("Skipping package '%s' due to missing asset data. Package may no longer exist!"), *AssetIdentifier.PackageName.ToString());
+				}
+			}
+		}
+	}
+}
+
+namespace PluginReferenceViewerCVars
+{
+	// Example usage: PluginReferenceViewer.ExportGraph SomePlugin SomePlugin.csv
+	static FAutoConsoleCommand ExportPluginReferenceGraph(
+		TEXT("PluginReferenceViewer.ExportGraph"),
+		TEXT("Exports to .csv the number of references (by type) that a plugin has for each of it's dependencies.\n"
+			"1st arg: single plugin name or multiple names seperated with ','.\n"
+			"2nd arg (optional): output filename.\n"
+			"Example: PluginReferenceViewer.ExportGraph PluginA,PluginB,PluginC PluginReport.csv"),
+		FConsoleCommandWithArgsDelegate::CreateStatic(PluginReferenceViewerUtils::ExportPluginReferenceGraph)
+	);
+}
+
+/*static*/ void FPluginReferenceViewerUtils::ExportPlugins(const TArray<FString>& InPluginNames, const FString& InFilename)
+{
+	TArray<TSharedRef<IPlugin>> Plugins;
+	for (const FString& PluginName : InPluginNames)
+	{
+		const TSharedPtr<IPlugin> Plugin = IPluginManager::Get().FindPlugin(PluginName);
+		if (Plugin.IsValid())
+		{
+			Plugins.Add(Plugin.ToSharedRef());
+		}
+	}
+
+	if (Plugins.IsEmpty())
+	{
+		UE_LOG(LogPluginReferenceViewerUtils, Display, TEXT("Plugin names array is empty"), *InFilename);
+		return;
+	}
+
+	if (FPaths::GetExtension(InFilename) != TEXT("csv"))
+	{
+		UE_LOG(LogPluginReferenceViewerUtils, Display, TEXT("Invalid filename extenstion '%s'. Expected .csv"), *InFilename);
+		return;
+	}
+
+	IAssetRegistry* AssetRegistry = IAssetRegistry::Get();
+	AssetRegistry->WaitForCompletion();
+
+	TUniquePtr<IFileHandle> ExportFileHandle(FPlatformFileManager::Get().GetPlatformFile().OpenWrite(*InFilename));
+	if (!ExportFileHandle.IsValid())
+	{
+		return;
+	}
+
+	UTF16CHAR BOM = UNICODE_BOM;
+	ExportFileHandle->Write((uint8*)&BOM, sizeof(UTF16CHAR));
+
+	constexpr TCHAR Separator = TEXT(',');
+	constexpr TCHAR LineEnd = TEXT('\n');
+
+	TStringBuilder<1024> StringBuilder;
+	
+	StringBuilder.Append(TEXT("Plugin, Dependency, Enabled, Optional, Asset References, Script References, Name References, Total References"));
+	StringBuilder.AppendChar(LineEnd);
+
+	ExportFileHandle->Write((const uint8*)StringBuilder.ToString(), StringBuilder.Len() * sizeof(TCHAR));
+	StringBuilder.Reset();
+
+	FScopedSlowTask SlowTask(Plugins.Num(), LOCTEXT("Exporting Plugin Graph", "Exporting Plugin Graph..."));
+	SlowTask.MakeDialog(true);
+
+	for (const TSharedRef<IPlugin>& CurrentPlugin : Plugins)
+	{
+		const TArray<FAssetIdentifier> AllDependencies = PluginReferenceViewerUtils::GetAssetDependencies(CurrentPlugin);
+		const TMap<FString, TArray<FAssetIdentifier>> PluginAssetMap = PluginReferenceViewerUtils::SplitByPlugins(CurrentPlugin, AllDependencies);
+
+		for (const FPluginReferenceDescriptor& ReferenceDescriptor : CurrentPlugin->GetDescriptor().Plugins)
+		{
+			// There might not be any asset references in the plugin dependency.
+			const TArray<FAssetIdentifier>* PluginAssets = PluginAssetMap.Find(ReferenceDescriptor.Name);
+			if (PluginAssets == nullptr)
+			{
+				continue;
+			}
+
+			TArray<FAssetIdentifier> AssetReferences, ScriptReferences, NameReferences;
+			PluginReferenceViewerUtils::SplitByReferenceType(*PluginAssets, AssetReferences, ScriptReferences, NameReferences);
+
+			StringBuilder.Reset();
+
+			StringBuilder.Append(CurrentPlugin->GetName());
+			StringBuilder.AppendChar(Separator);
+
+			StringBuilder.Append(ReferenceDescriptor.Name);
+			StringBuilder.AppendChar(Separator);
+
+			StringBuilder.Append(ReferenceDescriptor.bEnabled ? TEXT("true") : TEXT("false"));
+			StringBuilder.AppendChar(Separator);
+
+			StringBuilder.Append(ReferenceDescriptor.bOptional ? TEXT("true") : TEXT("false"));
+			StringBuilder.AppendChar(Separator);
+
+			StringBuilder.Appendf(TEXT("%d"), AssetReferences.Num());
+			StringBuilder.AppendChar(Separator);
+
+			StringBuilder.Appendf(TEXT("%d"), ScriptReferences.Num());
+			StringBuilder.AppendChar(Separator);
+
+			StringBuilder.Appendf(TEXT("%d"), NameReferences.Num());
+			StringBuilder.AppendChar(Separator);
+
+			StringBuilder.Appendf(TEXT("%d"), (*PluginAssets).Num());
+			StringBuilder.AppendChar(Separator);
+
+			StringBuilder.AppendChar(LineEnd);
+			ExportFileHandle->Write((const uint8*)StringBuilder.ToString(), StringBuilder.Len() * sizeof(TCHAR));
+		}
+	}
+
+	ExportFileHandle->Flush();
+
+
+	FStringBuilderBase ConcatenatePluginNames;
+	for (int32 Index = 0; Index < InPluginNames.Num(); ++Index)
+	{
+		ConcatenatePluginNames.Append(InPluginNames[Index]);
+		if (Index < InPluginNames.Num() - 1)
+		{
+			ConcatenatePluginNames.Append(TEXT(", "));
+		}
+	}
+
+	UE_LOG(LogPluginReferenceViewerUtils, Display, TEXT("Exported plugins; '%s' to '%s'"), *ConcatenatePluginNames, *FPaths::ConvertRelativePathToFull(InFilename));
+}
+
+#undef LOCTEXT_NAMESPACE
\ No newline at end of file
diff --git a/Engine/Plugins/Experimental/PluginReferenceViewer/Source/PluginReferenceViewer/Public/PluginReferenceViewerUtils.h b/Engine/Plugins/Experimental/PluginReferenceViewer/Source/PluginReferenceViewer/Public/PluginReferenceViewerUtils.h
new file mode 100644
--- /dev/null
+++ b/Engine/Plugins/Experimental/PluginReferenceViewer/Source/PluginReferenceViewer/Public/PluginReferenceViewerUtils.h
@@ -0,0 +1,16 @@
+// Copyright Epic Games, Inc. All Rights Reserved.
+
+#pragma once
+
+#include "CoreMinimal.h"
+
+class FPluginReferenceViewerUtils
+{
+public:
+	/* Exports the number of references from each plugin to their plugin dependencies. References include; assets, scripts and named references*/
+	static void ExportPlugins(const TArray<FString>& InPlugins, const FString& InFilename);
+
+private:
+	FPluginReferenceViewerUtils() = delete;
+	FPluginReferenceViewerUtils(const FPluginReferenceViewerUtils&) = delete;
+};
\ No newline at end of file
