Exports to a .csv file the list of asset references that exist between the plugin and one of it's plugin dependencies

diff --git a/Engine/Plugins/Experimental/PluginReferenceViewer/Source/PluginReferenceViewer/Private/PluginReferenceViewerUtils.cpp b/Engine/Plugins/Experimental/PluginReferenceViewer/Source/PluginReferenceViewer/Private/PluginReferenceViewerUtils.cpp
--- a/Engine/Plugins/Experimental/PluginReferenceViewer/Source/PluginReferenceViewer/Private/PluginReferenceViewerUtils.cpp
+++ b/Engine/Plugins/Experimental/PluginReferenceViewer/Source/PluginReferenceViewer/Private/PluginReferenceViewerUtils.cpp
@@ -22,7 +22,7 @@ DEFINE_LOG_CATEGORY_STATIC(LogPluginReferenceViewerUtils, Log, All);
 
 namespace PluginReferenceViewerUtils
 {
-	void ExportPluginReferenceGraph(const TArray<FString>& InArgs)
+	void ExportGraph(const TArray<FString>& InArgs)
 	{
 		TArray<FString> PluginNames;
 		if (InArgs.Num() >= 1)
@@ -48,58 +48,28 @@ namespace PluginReferenceViewerUtils
 		FPluginReferenceViewerUtils::ExportPlugins(PluginNames, Filename);
 	}
 
-	TArray<TSharedRef<const IPlugin>> GetGameplayTagSourcePlugins(FName TagName)
+	void ExportReference(const TArray<FString>& InArgs)
 	{
-		TArray<TSharedRef<const IPlugin>> Result;
-		UGameplayTagsManager& Manager = UGameplayTagsManager::Get();
+		if (InArgs.Num() < 2)
+		{
+			UE_LOG(LogPluginReferenceViewerUtils, Display, TEXT("Invalid arguments. Expected: [plugin name] [reference name] (optional)[filename.csv]"));
+			return;
+		}
+
+		FString PluginName = InArgs[0];
+		FString ReferenceName = InArgs[1];
 
-		FString Comment;
-		TArray<FName> TagSources;
-		bool bIsTagExplicit, bIsRestrictedTag, bAllowNonRestrictedChildren;
-		if (Manager.GetTagEditorData(TagName, Comment, TagSources, bIsTagExplicit, bIsRestrictedTag, bAllowNonRestrictedChildren))
+		FString Filename;
+		if (InArgs.Num() >= 3)
 		{
-			for (const FName& TagSourceName : TagSources)
-			{
-				if (const FGameplayTagSource* TagSource = Manager.FindTagSource(TagSourceName))
-				{
-					switch (TagSource->SourceType)
-					{
-					case EGameplayTagSourceType::TagList:
-					{
-						const FString ContentFilePath = FPaths::GetPath(TagSource->SourceTagList->ConfigFileName) / TEXT("../../Content/");
-						FString RootContentPath;
-						if (FPackageName::TryConvertFilenameToLongPackageName(ContentFilePath, RootContentPath))
-						{
-							if (const TSharedPtr<IPlugin> FoundPlugin = IPluginManager::Get().FindPluginFromPath(RootContentPath))
-							{
-								Result.Add(FoundPlugin.ToSharedRef());
-							}
-						}
-						break;
-					}
-					case EGameplayTagSourceType::DataTable:
-					{
-						if (const TSharedPtr<IPlugin>& FoundPlugin = IPluginManager::Get().FindPluginFromPath(TagSource->SourceName.ToString()))
-						{
-							Result.Add(FoundPlugin.ToSharedRef());
-						}
-						break;
-					}
-					case EGameplayTagSourceType::Native:
-					{
-						if (const TSharedPtr<IPlugin>& FoundPlugin = IPluginManager::Get().GetModuleOwnerPlugin(TagSource->SourceName))
-						{
-							Result.Add(FoundPlugin.ToSharedRef());
-						}
-						break;
-					}
-					default:
-						break;
-					}
-				}
-			}
+			Filename = InArgs[2];
+		}
+		else
+		{
+			Filename = FPaths::SetExtension(FString::Printf(TEXT("%s-%s"), *PluginName, *ReferenceName), TEXT("csv"));
 		}
-		return Result;
+
+		FPluginReferenceViewerUtils::ExportReference(PluginName, ReferenceName, Filename);
 	}
 
 	TArray<FAssetIdentifier> GetAssetDependencies(const TSharedRef<IPlugin>& InPlugin)
@@ -142,8 +112,8 @@ namespace PluginReferenceViewerUtils
 
 			if (AssetIdentifier.ObjectName == NAME_GameplayTag && AssetIdentifier.PackageName == GameplayTagStructPackage)
 			{
-				const TArray<TSharedRef<const IPlugin>> SourcePlugins = PluginReferenceViewerUtils::GetGameplayTagSourcePlugins(AssetIdentifier.ValueName);
-				for (const TSharedRef<const IPlugin>& SourcePlugin : SourcePlugins)
+				const TArray<TSharedRef<IPlugin>> SourcePlugins = FPluginReferenceViewerUtils::FindGameplayTagSourcePlugins(AssetIdentifier.ValueName);
+				for (const TSharedRef<IPlugin>& SourcePlugin : SourcePlugins)
 				{
 					if (SourcePlugin != InOwningPlugin)
 					{
@@ -229,14 +199,23 @@ namespace PluginReferenceViewerUtils
 
 namespace PluginReferenceViewerCVars
 {
-	// Example usage: PluginReferenceViewer.ExportGraph SomePlugin SomePlugin.csv
-	static FAutoConsoleCommand ExportPluginReferenceGraph(
+	static FAutoConsoleCommand ExportGraph(
 		TEXT("PluginReferenceViewer.ExportGraph"),
 		TEXT("Exports to .csv the number of references (by type) that a plugin has for each of it's dependencies.\n"
 			"1st arg: single plugin name or multiple names seperated with ','.\n"
 			"2nd arg (optional): output filename.\n"
 			"Example: PluginReferenceViewer.ExportGraph PluginA,PluginB,PluginC PluginReport.csv"),
-		FConsoleCommandWithArgsDelegate::CreateStatic(PluginReferenceViewerUtils::ExportPluginReferenceGraph)
+		FConsoleCommandWithArgsDelegate::CreateStatic(PluginReferenceViewerUtils::ExportGraph)
+	);
+
+	static FAutoConsoleCommand ExportReference(
+		TEXT("PluginReferenceViewer.ExportReference"),
+		TEXT("Exports to .csv the list of asset references that exist between a plugin and one of it's dependencies.\n"
+				"1st arg: plugin name.\n"
+				"2nd arg: reference name.\n"
+				"2rd arg (optional): output filename.\n"
+				"Example: PluginReferenceViewer.ExportReference PluginName, ReferenceName PluginReport.csv"),
+		FConsoleCommandWithArgsDelegate::CreateStatic(PluginReferenceViewerUtils::ExportReference)
 	);
 }
 
@@ -354,4 +333,123 @@ namespace PluginReferenceViewerCVars
 	UE_LOG(LogPluginReferenceViewerUtils, Display, TEXT("Exported plugins; '%s' to '%s'"), *ConcatenatePluginNames, *FPaths::ConvertRelativePathToFull(InFilename));
 }
 
+/*static*/ void FPluginReferenceViewerUtils::ExportReference(const FString& InPlugin, const FString& InReference, const FString& InFilename)
+{
+	const TSharedPtr<IPlugin> Plugin = IPluginManager::Get().FindPlugin(InPlugin);
+	if (!Plugin.IsValid())
+	{
+		UE_LOG(LogPluginReferenceViewerUtils, Display, TEXT("Plugin %s was not found!"), *InPlugin);
+		return;
+	}
+
+	const FPluginReferenceDescriptor* Found = Plugin->GetDescriptor().Plugins.FindByPredicate([=](const FPluginReferenceDescriptor& Item) { return Item.Name == InReference; });
+	if (Found == nullptr)
+	{
+		UE_LOG(LogPluginReferenceViewerUtils, Display, TEXT("Plugin reference %s was not found!"), *InReference);
+		return;
+	}
+
+	if (FPaths::GetExtension(InFilename) != TEXT("csv"))
+	{
+		UE_LOG(LogPluginReferenceViewerUtils, Display, TEXT("Invalid filename extenstion '%s'. Expected .csv"), *InFilename);
+		return;
+	}
+
+	IAssetRegistry* AssetRegistry = IAssetRegistry::Get();
+	AssetRegistry->WaitForCompletion();
+
+	TUniquePtr<IFileHandle> ExportFileHandle(FPlatformFileManager::Get().GetPlatformFile().OpenWrite(*InFilename));
+	if (!ExportFileHandle.IsValid())
+	{
+		return;
+	}
+
+	UTF16CHAR BOM = UNICODE_BOM;
+	ExportFileHandle->Write((uint8*)&BOM, sizeof(UTF16CHAR));
+
+	constexpr TCHAR Separator = TEXT(',');
+	constexpr TCHAR LineEnd = TEXT('\n');
+
+	TStringBuilder<1024> StringBuilder;
+
+	const TArray<FAssetIdentifier> AllDependencies = PluginReferenceViewerUtils::GetAssetDependencies(Plugin.ToSharedRef());
+	const TMap<FString, TArray<FAssetIdentifier>> PluginAssetMap = PluginReferenceViewerUtils::SplitByPlugins(Plugin.ToSharedRef(), AllDependencies);
+
+	// There might not be any asset references in the plugin dependency.
+	const TArray<FAssetIdentifier>* PluginAssets = PluginAssetMap.Find(InReference);
+	if (PluginAssets != nullptr)
+	{
+		TArray<FAssetIdentifier> AssetReferences, ScriptReferences, NameReferences;
+		PluginReferenceViewerUtils::SplitByReferenceType(*PluginAssets, AssetReferences, ScriptReferences, NameReferences);
+
+		for (TArray<FAssetIdentifier>* AssetList : { &AssetReferences, &ScriptReferences, &NameReferences })
+		{
+			for (const FAssetIdentifier& Identifier : *AssetList)
+			{
+				StringBuilder.Reset();
+
+				StringBuilder.Append(Identifier.ToString());
+				StringBuilder.AppendChar(LineEnd);
+				ExportFileHandle->Write((const uint8*)StringBuilder.ToString(), StringBuilder.Len() * sizeof(TCHAR));
+			}
+		}
+	}
+
+	ExportFileHandle->Flush();
+}
+
+/*static*/ TArray<TSharedRef<IPlugin>> FPluginReferenceViewerUtils::FindGameplayTagSourcePlugins(FName TagName)
+{
+	TArray<TSharedRef<IPlugin>> Result;
+	UGameplayTagsManager& Manager = UGameplayTagsManager::Get();
+
+	FString Comment;
+	TArray<FName> TagSources;
+	bool bIsTagExplicit, bIsRestrictedTag, bAllowNonRestrictedChildren;
+	if (Manager.GetTagEditorData(TagName, Comment, TagSources, bIsTagExplicit, bIsRestrictedTag, bAllowNonRestrictedChildren))
+	{
+		for (const FName& TagSourceName : TagSources)
+		{
+			if (const FGameplayTagSource* TagSource = Manager.FindTagSource(TagSourceName))
+			{
+				switch (TagSource->SourceType)
+				{
+				case EGameplayTagSourceType::TagList:
+				{
+					const FString ContentFilePath = FPaths::GetPath(TagSource->SourceTagList->ConfigFileName) / TEXT("../../Content/");
+					FString RootContentPath;
+					if (FPackageName::TryConvertFilenameToLongPackageName(ContentFilePath, RootContentPath))
+					{
+						if (const TSharedPtr<IPlugin> FoundPlugin = IPluginManager::Get().FindPluginFromPath(RootContentPath))
+						{
+							Result.Add(FoundPlugin.ToSharedRef());
+						}
+					}
+					break;
+				}
+				case EGameplayTagSourceType::DataTable:
+				{
+					if (const TSharedPtr<IPlugin>& FoundPlugin = IPluginManager::Get().FindPluginFromPath(TagSource->SourceName.ToString()))
+					{
+						Result.Add(FoundPlugin.ToSharedRef());
+					}
+					break;
+				}
+				case EGameplayTagSourceType::Native:
+				{
+					if (const TSharedPtr<IPlugin>& FoundPlugin = IPluginManager::Get().GetModuleOwnerPlugin(TagSource->SourceName))
+					{
+						Result.Add(FoundPlugin.ToSharedRef());
+					}
+					break;
+				}
+				default:
+					break;
+				}
+			}
+		}
+	}
+	return Result;
+}
+
 #undef LOCTEXT_NAMESPACE
\ No newline at end of file
diff --git a/Engine/Plugins/Experimental/PluginReferenceViewer/Source/PluginReferenceViewer/Public/PluginReferenceViewerUtils.h b/Engine/Plugins/Experimental/PluginReferenceViewer/Source/PluginReferenceViewer/Public/PluginReferenceViewerUtils.h
--- a/Engine/Plugins/Experimental/PluginReferenceViewer/Source/PluginReferenceViewer/Public/PluginReferenceViewerUtils.h
+++ b/Engine/Plugins/Experimental/PluginReferenceViewer/Source/PluginReferenceViewer/Public/PluginReferenceViewerUtils.h
@@ -4,12 +4,20 @@
 
 #include "CoreMinimal.h"
 
-class FPluginReferenceViewerUtils
+class IPlugin;
+
+class PLUGINREFERENCEVIEWER_API FPluginReferenceViewerUtils
 {
 public:
 	/* Exports the number of references from each plugin to their plugin dependencies. References include; assets, scripts and named references*/
 	static void ExportPlugins(const TArray<FString>& InPlugins, const FString& InFilename);
 
+	/* Exports to a .csv file the list of asset references that exist between the plugin and one of it's plugin dependencies */
+	static void ExportReference(const FString& InPlugin, const FString& InReference, const FString& InFilename);
+
+	/* Returns the list of plugins where the supplied gameplay tag is declated */
+	static TArray<TSharedRef<IPlugin>> FindGameplayTagSourcePlugins(FName TagName);
+
 private:
 	FPluginReferenceViewerUtils() = delete;
 	FPluginReferenceViewerUtils(const FPluginReferenceViewerUtils&) = delete;
