ExportDirectory command added what will export to csv all cross plugin asset references for all plugins found under the specified directory.

diff --git a/Engine/Plugins/Experimental/PluginReferenceViewer/Source/PluginReferenceViewer/Private/PluginReferenceViewerUtils.cpp b/Engine/Plugins/Experimental/PluginReferenceViewer/Source/PluginReferenceViewer/Private/PluginReferenceViewerUtils.cpp
--- a/Engine/Plugins/Experimental/PluginReferenceViewer/Source/PluginReferenceViewer/Private/PluginReferenceViewerUtils.cpp
+++ b/Engine/Plugins/Experimental/PluginReferenceViewer/Source/PluginReferenceViewer/Private/PluginReferenceViewerUtils.cpp
@@ -22,7 +22,7 @@ DEFINE_LOG_CATEGORY_STATIC(LogPluginReferenceViewerUtils, Log, All);
 
 namespace PluginReferenceViewerUtils
 {
-	void ExportGraph(const TArray<FString>& InArgs)
+	void ExportPlugins(const TArray<FString>& InArgs)
 	{
 		TArray<FString> PluginNames;
 		if (InArgs.Num() >= 1)
@@ -31,7 +31,7 @@ namespace PluginReferenceViewerUtils
 		}
 		else
 		{
-			UE_LOG(LogPluginReferenceViewerUtils, Display, TEXT("Invalid plugin names argument. Expect plugin names seperated by ',' as the first argument. e.g a,b,c"));
+			UE_LOG(LogPluginReferenceViewerUtils, Error, TEXT("Invalid plugin names argument. Expect plugin names seperated by ',' as the first argument. e.g a,b,c"));
 			return;
 		}
 
@@ -42,7 +42,7 @@ namespace PluginReferenceViewerUtils
 		}
 		else if (!PluginNames.IsEmpty())
 		{
-			Filename = FPaths::SetExtension(PluginNames[0], TEXT("csv"));
+			Filename = FPaths::ProjectSavedDir() / FPaths::SetExtension(PluginNames[0], TEXT("csv"));
 		}
 
 		FPluginReferenceViewerUtils::ExportPlugins(PluginNames, Filename);
@@ -52,7 +52,7 @@ namespace PluginReferenceViewerUtils
 	{
 		if (InArgs.Num() < 2)
 		{
-			UE_LOG(LogPluginReferenceViewerUtils, Display, TEXT("Invalid arguments. Expected: [plugin name] [reference name] (optional)[filename.csv]"));
+			UE_LOG(LogPluginReferenceViewerUtils, Error, TEXT("Invalid arguments. Expected: [plugin name] [reference name] (optional)[filename.csv]"));
 			return;
 		}
 
@@ -66,12 +66,51 @@ namespace PluginReferenceViewerUtils
 		}
 		else
 		{
-			Filename = FPaths::SetExtension(FString::Printf(TEXT("%s-%s"), *PluginName, *ReferenceName), TEXT("csv"));
+			Filename = FPaths::ProjectSavedDir() / FPaths::SetExtension(FString::Printf(TEXT("%s-%s"), *PluginName, *ReferenceName), TEXT("csv"));
 		}
 
 		FPluginReferenceViewerUtils::ExportReference(PluginName, ReferenceName, Filename);
 	}
 
+	void ExportDirectory(const TArray<FString>& InArgs)
+	{
+		TArray<FString> PluginNames;
+		if (InArgs.Num() >= 1)
+		{
+			const FString SearchDirectory = FPaths::RootDir() / InArgs[0];
+			if (FPaths::DirectoryExists(SearchDirectory))
+			{
+				TArray<FString> PluginFileNames;
+				IPluginManager::Get().FindPluginsUnderDirectory(SearchDirectory, PluginFileNames);
+
+				Algo::Transform(PluginFileNames, PluginNames, [](const FString& Item) { return FPaths::GetBaseFilename(Item); });
+				PluginNames.Sort();
+			}
+			else
+			{
+				UE_LOG(LogPluginReferenceViewerUtils, Error, TEXT("Directory does not exist %s"), *SearchDirectory);
+				return;
+			}
+		}
+		else
+		{
+			UE_LOG(LogPluginReferenceViewerUtils, Error, TEXT("Invalid arguments. Expected directory path"));
+			return;
+		}
+
+		FString Filename;
+		if (InArgs.Num() >= 2)
+		{
+			Filename = InArgs[1];
+		}
+		else if (!PluginNames.IsEmpty())
+		{
+			Filename = FPaths::ProjectSavedDir() / FPaths::SetExtension(FPaths::GetBaseFilename(InArgs[0]), TEXT("csv"));
+		}
+
+		FPluginReferenceViewerUtils::ExportPlugins(PluginNames, Filename);
+	}
+
 	TArray<FAssetIdentifier> GetAssetDependencies(const TSharedRef<IPlugin>& InPlugin)
 	{
 		IAssetRegistry& AssetRegistry = FModuleManager::LoadModuleChecked<FAssetRegistryModule>("AssetRegistry").Get();
@@ -199,13 +238,13 @@ namespace PluginReferenceViewerUtils
 
 namespace PluginReferenceViewerCVars
 {
-	static FAutoConsoleCommand ExportGraph(
-		TEXT("PluginReferenceViewer.ExportGraph"),
+	static FAutoConsoleCommand ExportPlugins(
+		TEXT("PluginReferenceViewer.ExportPlugins"),
 		TEXT("Exports to .csv the number of references (by type) that a plugin has for each of it's dependencies.\n"
 			"1st arg: single plugin name or multiple names seperated with ','.\n"
 			"2nd arg (optional): output filename.\n"
-			"Example: PluginReferenceViewer.ExportGraph PluginA,PluginB,PluginC PluginReport.csv"),
-		FConsoleCommandWithArgsDelegate::CreateStatic(PluginReferenceViewerUtils::ExportGraph)
+			"Example: PluginReferenceViewer.ExportPlugins PluginA,PluginB,PluginC PluginReport.csv"),
+		FConsoleCommandWithArgsDelegate::CreateStatic(PluginReferenceViewerUtils::ExportPlugins)
 	);
 
 	static FAutoConsoleCommand ExportReference(
@@ -213,10 +252,19 @@ namespace PluginReferenceViewerCVars
 		TEXT("Exports to .csv the list of asset references that exist between a plugin and one of it's dependencies.\n"
 				"1st arg: plugin name.\n"
 				"2nd arg: reference name.\n"
-				"2rd arg (optional): output filename.\n"
-				"Example: PluginReferenceViewer.ExportReference PluginName, ReferenceName PluginReport.csv"),
+				"3rd arg (optional): output filename.\n"
+				"Example: PluginReferenceViewer.ExportReference PluginName ReferenceName PluginReport.csv"),
 		FConsoleCommandWithArgsDelegate::CreateStatic(PluginReferenceViewerUtils::ExportReference)
 	);
+
+	static FAutoConsoleCommand ExportDirectory(
+		TEXT("PluginReferenceViewer.ExportDirectory"),
+		TEXT("Exports to .csv the list of asset references that exist between each found plugin and all dependencies.\n"
+			"1st arg: path relative to the root directory.\n"
+			"2rd arg (optional): output filename.\n"
+			"Example: PluginReferenceViewer.ExportDirectory Path PluginReport.csv"),
+		FConsoleCommandWithArgsDelegate::CreateStatic(PluginReferenceViewerUtils::ExportDirectory)
+	);
 }
 
 /*static*/ void FPluginReferenceViewerUtils::ExportPlugins(const TArray<FString>& InPluginNames, const FString& InFilename)
@@ -233,13 +281,13 @@ namespace PluginReferenceViewerCVars
 
 	if (Plugins.IsEmpty())
 	{
-		UE_LOG(LogPluginReferenceViewerUtils, Display, TEXT("Plugin names array is empty"), *InFilename);
+		UE_LOG(LogPluginReferenceViewerUtils, Error, TEXT("Plugin names array is empty"), *InFilename);
 		return;
 	}
 
 	if (FPaths::GetExtension(InFilename) != TEXT("csv"))
 	{
-		UE_LOG(LogPluginReferenceViewerUtils, Display, TEXT("Invalid filename extenstion '%s'. Expected .csv"), *InFilename);
+		UE_LOG(LogPluginReferenceViewerUtils, Error, TEXT("Invalid filename extenstion '%s'. Expected .csv"), *InFilename);
 		return;
 	}
 
@@ -271,6 +319,8 @@ namespace PluginReferenceViewerCVars
 
 	for (const TSharedRef<IPlugin>& CurrentPlugin : Plugins)
 	{
+		SlowTask.EnterProgressFrame(1.0f, FText::Format(LOCTEXT("ExportPluginName", "Processing plugin {0}"), FText::FromString(CurrentPlugin->GetName())));
+
 		const TArray<FAssetIdentifier> AllDependencies = PluginReferenceViewerUtils::GetAssetDependencies(CurrentPlugin);
 		const TMap<FString, TArray<FAssetIdentifier>> PluginAssetMap = PluginReferenceViewerUtils::SplitByPlugins(CurrentPlugin, AllDependencies);
 
@@ -338,20 +388,20 @@ namespace PluginReferenceViewerCVars
 	const TSharedPtr<IPlugin> Plugin = IPluginManager::Get().FindPlugin(InPlugin);
 	if (!Plugin.IsValid())
 	{
-		UE_LOG(LogPluginReferenceViewerUtils, Display, TEXT("Plugin %s was not found!"), *InPlugin);
+		UE_LOG(LogPluginReferenceViewerUtils, Error, TEXT("Plugin %s was not found!"), *InPlugin);
 		return;
 	}
 
 	const FPluginReferenceDescriptor* Found = Plugin->GetDescriptor().Plugins.FindByPredicate([=](const FPluginReferenceDescriptor& Item) { return Item.Name == InReference; });
 	if (Found == nullptr)
 	{
-		UE_LOG(LogPluginReferenceViewerUtils, Display, TEXT("Plugin reference %s was not found!"), *InReference);
+		UE_LOG(LogPluginReferenceViewerUtils, Error, TEXT("Plugin reference %s was not found!"), *InReference);
 		return;
 	}
 
 	if (FPaths::GetExtension(InFilename) != TEXT("csv"))
 	{
-		UE_LOG(LogPluginReferenceViewerUtils, Display, TEXT("Invalid filename extenstion '%s'. Expected .csv"), *InFilename);
+		UE_LOG(LogPluginReferenceViewerUtils, Error, TEXT("Invalid filename extenstion '%s'. Expected .csv"), *InFilename);
 		return;
 	}
 
