Utility console command to trace a plugin dependency chain and output the found paths

diff --git a/Engine/Plugins/Experimental/PluginReferenceViewer/Source/PluginReferenceViewer/Private/PluginReferenceViewerUtils.cpp b/Engine/Plugins/Experimental/PluginReferenceViewer/Source/PluginReferenceViewer/Private/PluginReferenceViewerUtils.cpp
--- a/Engine/Plugins/Experimental/PluginReferenceViewer/Source/PluginReferenceViewer/Private/PluginReferenceViewerUtils.cpp
+++ b/Engine/Plugins/Experimental/PluginReferenceViewer/Source/PluginReferenceViewer/Private/PluginReferenceViewerUtils.cpp
@@ -111,6 +111,108 @@ namespace PluginReferenceViewerUtils
 		FPluginReferenceViewerUtils::ExportPlugins(PluginNames, Filename);
 	}
 
+	void TracePath(const TArray<FString>& InArgs)
+	{
+		FString StartPluginName;
+		if (InArgs.Num() >= 1)
+		{
+			StartPluginName = InArgs[0];
+		}
+		else
+		{
+			UE_LOG(LogPluginReferenceViewerUtils, Error, TEXT("Invalid arguments. Expected plugin name as 1st arg"));
+			return;
+		}
+
+		FString EndPluginName;
+		if (InArgs.Num() >= 2)
+		{
+			EndPluginName = InArgs[1];
+		}
+		else
+		{
+			UE_LOG(LogPluginReferenceViewerUtils, Error, TEXT("Invalid arguments. Expected plugin name as 2nd arg"));
+			return;
+		}
+
+		const TSharedPtr<IPlugin> Plugin = IPluginManager::Get().FindPlugin(StartPluginName);
+		if (!Plugin.IsValid())
+		{
+			UE_LOG(LogPluginReferenceViewerUtils, Error, TEXT("Plugin `%s` could not be found!"), *StartPluginName);
+			return;
+		}
+
+		struct FPluginVisitor
+		{
+			FString StartPoint;
+			FString EndPoint;
+			TArray<FString> Stack;
+			bool bPathFound = false;
+
+			FPluginVisitor(const FString& InStartPoint, const FString& InEndPoint)
+				: StartPoint(InStartPoint)
+				, EndPoint(InEndPoint)
+			{
+			}
+
+			void VisitPlugins()
+			{
+				Stack.Push(StartPoint);
+				VisitPluginRecursive(StartPoint);
+				Stack.Pop();
+
+				if (!bPathFound)
+				{
+					UE_LOG(LogPluginReferenceViewerUtils, Display, TEXT("No paths from plugin `%s` to plugin '%s' was found!"), *StartPoint, *EndPoint);
+				}
+			}
+
+		private:
+			void VisitPluginRecursive(const FString& InName)
+			{
+				if (InName == EndPoint)
+				{
+					bPathFound = true;
+
+					TraceStack();
+				}
+				else
+				{
+					const TSharedPtr<IPlugin> Plugin = IPluginManager::Get().FindPlugin(InName);
+					if (Plugin.IsValid())
+					{
+						for (const FPluginReferenceDescriptor& PluginDescriptor : Plugin->GetDescriptor().Plugins)
+						{
+							Stack.Push(PluginDescriptor.Name);
+							VisitPluginRecursive(PluginDescriptor.Name);
+							Stack.Pop();
+						}
+					}
+				}
+			}
+
+			void TraceStack()
+			{
+				FStringBuilderBase PathString;
+				PathString.Append(FString::Format(TEXT("Found dependency path of length {0} : "), { Stack.Num() }));
+
+				for (int32 Index = 0; Index < Stack.Num(); ++Index)
+				{
+					PathString.Append(*Stack[Index]);
+					if (Index < Stack.Num() - 1)
+					{
+						PathString.Append(TEXT(" -> "));
+					}
+				}
+
+				UE_LOG(LogPluginReferenceViewerUtils, Display, TEXT("%s"), PathString.ToString());
+			}
+		};
+
+		FPluginVisitor PluginVisitor(StartPluginName, EndPluginName);
+		PluginVisitor.VisitPlugins();
+	}
+
 	TArray<FAssetIdentifier> GetAssetDependencies(const TSharedRef<IPlugin>& InPlugin)
 	{
 		IAssetRegistry& AssetRegistry = FModuleManager::LoadModuleChecked<FAssetRegistryModule>("AssetRegistry").Get();
@@ -265,6 +367,15 @@ namespace PluginReferenceViewerCVars
 			"Example: PluginReferenceViewer.ExportDirectory Path PluginReport.csv"),
 		FConsoleCommandWithArgsDelegate::CreateStatic(PluginReferenceViewerUtils::ExportDirectory)
 	);
+
+	static FAutoConsoleCommand TracePath(
+		TEXT("PluginReferenceViewer.TracePath"),
+		TEXT("Outputs all found plugin dependency paths from plugin X to plugin Y.\n"
+			"1st arg: plugin start point name.\n"
+			"2rd arg: plugin end point name.\n"
+			"Example: PluginReferenceViewer.TracePath PluginX PluginY"),
+		FConsoleCommandWithArgsDelegate::CreateStatic(PluginReferenceViewerUtils::TracePath)
+	);
 }
 
 /*static*/ void FPluginReferenceViewerUtils::ExportPlugins(const TArray<FString>& InPluginNames, const FString& InFilename)
