Fix issue where DisallowedPlugins wasn't correctly being parsed out during build tool
Add new method users can use to determine if a plugin chain from one plugin to another exists.
Update implementation to reduce the amount of string allocations when finding the path from one plugin to another

diff --git a/Engine/Plugins/Experimental/PluginReferenceViewer/Source/PluginReferenceViewer/Private/PluginReferenceViewerUtils.cpp b/Engine/Plugins/Experimental/PluginReferenceViewer/Source/PluginReferenceViewer/Private/PluginReferenceViewerUtils.cpp
--- a/Engine/Plugins/Experimental/PluginReferenceViewer/Source/PluginReferenceViewer/Private/PluginReferenceViewerUtils.cpp
+++ b/Engine/Plugins/Experimental/PluginReferenceViewer/Source/PluginReferenceViewer/Private/PluginReferenceViewerUtils.cpp
@@ -2,6 +2,7 @@
 
 #include "PluginReferenceViewerUtils.h"
 
+#include "Algo/Reverse.h"
 #include "AssetRegistry/AssetRegistryModule.h"
 #include "AssetRegistry/AssetData.h"
 #include "AssetManagerEditorModule.h"
@@ -111,6 +112,63 @@ namespace PluginReferenceViewerUtils
 		FPluginReferenceViewerUtils::ExportPlugins(PluginNames, Filename);
 	}
 
+	static void GetPluginDependenciesRecursive_Helper(IPluginManager& PluginManager, const FString& ParentPluginName, TMap<FString, FString>& OutChildToParent)
+	{
+		TSharedPtr<IPlugin> Plugin = PluginManager.FindPlugin(ParentPluginName);
+		if (Plugin.IsValid())
+		{
+			const FPluginDescriptor& Desc = Plugin->GetDescriptor();
+			for (const FPluginReferenceDescriptor& Dependency : Desc.Plugins)
+			{
+				if (!OutChildToParent.Contains(Dependency.Name))
+				{
+					OutChildToParent.Add(Dependency.Name, ParentPluginName);
+					GetPluginDependenciesRecursive_Helper(PluginManager, Dependency.Name, OutChildToParent);
+				}
+			}
+		}
+	}
+
+	void GetPluginDependenciesRecursive(const FString& PluginName, TMap<FString, FString>& OutChildToParent)
+	{
+		GetPluginDependenciesRecursive_Helper(IPluginManager::Get(), PluginName, OutChildToParent);
+	}
+
+	bool TracePathFromPluginToPlugin(const FString& StartPluginName, const FString& EndPluginName, FString& FoundPathToEndPlugin)
+	{
+		const TSharedPtr<IPlugin> PluginPtr = IPluginManager::Get().FindPlugin(StartPluginName);
+		if (!PluginPtr.IsValid())
+		{
+			UE_LOG(LogPluginReferenceViewerUtils, Error, TEXT("Plugin `%s` could not be found!"), *StartPluginName);
+			return false;
+		}
+
+		TMap<FString, FString> ChildToParent;
+		GetPluginDependenciesRecursive(StartPluginName, ChildToParent);
+		if (!ChildToParent.Contains(EndPluginName))
+		{
+			UE_LOG(LogPluginReferenceViewerUtils, Display, TEXT("No paths from plugin `%s` to plugin '%s' was found!"), *StartPluginName, *EndPluginName);
+			return false;
+		}
+
+		TArray<FString> ReversePath;
+		FString CurrentPluginName = EndPluginName;
+		ReversePath.Add(CurrentPluginName);
+		while (CurrentPluginName != StartPluginName)
+		{
+			CurrentPluginName = ChildToParent[CurrentPluginName];
+			ReversePath.Add(CurrentPluginName);
+		}
+		Algo::Reverse(ReversePath); // Now correct path
+
+		FStringBuilderBase PathString;
+		PathString.Append(FString::Format(TEXT("Found dependency path of length {0} : "), { ReversePath.Num() }));
+		PathString.Append(FString::Join(ReversePath, TEXT(" -> ")));
+		FoundPathToEndPlugin = PathString.ToString();
+		UE_LOG(LogPluginReferenceViewerUtils, Display, TEXT("%s"), *FoundPathToEndPlugin);
+		return true;
+	}
+
 	void TracePath(const TArray<FString>& InArgs)
 	{
 		FString StartPluginName;
@@ -135,82 +193,8 @@ namespace PluginReferenceViewerUtils
 			return;
 		}
 
-		const TSharedPtr<IPlugin> Plugin = IPluginManager::Get().FindPlugin(StartPluginName);
-		if (!Plugin.IsValid())
-		{
-			UE_LOG(LogPluginReferenceViewerUtils, Error, TEXT("Plugin `%s` could not be found!"), *StartPluginName);
-			return;
-		}
-
-		struct FPluginVisitor
-		{
-			FString StartPoint;
-			FString EndPoint;
-			TArray<FString> Stack;
-			bool bPathFound = false;
-
-			FPluginVisitor(const FString& InStartPoint, const FString& InEndPoint)
-				: StartPoint(InStartPoint)
-				, EndPoint(InEndPoint)
-			{
-			}
-
-			void VisitPlugins()
-			{
-				Stack.Push(StartPoint);
-				VisitPluginRecursive(StartPoint);
-				Stack.Pop();
-
-				if (!bPathFound)
-				{
-					UE_LOG(LogPluginReferenceViewerUtils, Display, TEXT("No paths from plugin `%s` to plugin '%s' was found!"), *StartPoint, *EndPoint);
-				}
-			}
-
-		private:
-			void VisitPluginRecursive(const FString& InName)
-			{
-				if (InName == EndPoint)
-				{
-					bPathFound = true;
-
-					TraceStack();
-				}
-				else
-				{
-					const TSharedPtr<IPlugin> Plugin = IPluginManager::Get().FindPlugin(InName);
-					if (Plugin.IsValid())
-					{
-						for (const FPluginReferenceDescriptor& PluginDescriptor : Plugin->GetDescriptor().Plugins)
-						{
-							Stack.Push(PluginDescriptor.Name);
-							VisitPluginRecursive(PluginDescriptor.Name);
-							Stack.Pop();
-						}
-					}
-				}
-			}
-
-			void TraceStack()
-			{
-				FStringBuilderBase PathString;
-				PathString.Append(FString::Format(TEXT("Found dependency path of length {0} : "), { Stack.Num() }));
-
-				for (int32 Index = 0; Index < Stack.Num(); ++Index)
-				{
-					PathString.Append(*Stack[Index]);
-					if (Index < Stack.Num() - 1)
-					{
-						PathString.Append(TEXT(" -> "));
-					}
-				}
-
-				UE_LOG(LogPluginReferenceViewerUtils, Display, TEXT("%s"), PathString.ToString());
-			}
-		};
-
-		FPluginVisitor PluginVisitor(StartPluginName, EndPluginName);
-		PluginVisitor.VisitPlugins();
+		FString tmp;
+		TracePathFromPluginToPlugin(StartPluginName, EndPluginName, tmp);
 	}
 
 	TArray<FAssetIdentifier> GetAssetDependencies(const TSharedRef<IPlugin>& InPlugin)
@@ -613,4 +597,9 @@ namespace PluginReferenceViewerCVars
 	return Result;
 }
 
+bool FPluginReferenceViewerUtils::TracePluginChain(const FString& StartingPlugin, const FString& EndingPlugin, FString& OutPathToEndPlugin)
+{
+	return PluginReferenceViewerUtils::TracePathFromPluginToPlugin(StartingPlugin, EndingPlugin, OutPathToEndPlugin);
+}
+
 #undef LOCTEXT_NAMESPACE
\ No newline at end of file
diff --git a/Engine/Plugins/Experimental/PluginReferenceViewer/Source/PluginReferenceViewer/Public/PluginReferenceViewerUtils.h b/Engine/Plugins/Experimental/PluginReferenceViewer/Source/PluginReferenceViewer/Public/PluginReferenceViewerUtils.h
--- a/Engine/Plugins/Experimental/PluginReferenceViewer/Source/PluginReferenceViewer/Public/PluginReferenceViewerUtils.h
+++ b/Engine/Plugins/Experimental/PluginReferenceViewer/Source/PluginReferenceViewer/Public/PluginReferenceViewerUtils.h
@@ -18,6 +18,12 @@ public:
 	/* Returns the list of plugins where the supplied gameplay tag is declated */
 	static TArray<TSharedRef<IPlugin>> FindGameplayTagSourcePlugins(FName TagName);
 
+	/*
+	 * Traces the path from one plugin to the ending plugin if one exists.
+	 * Returns true if a path exists.
+	 */
+	static bool TracePluginChain(const FString& StartingPlugin, const FString& EndingPlugin, FString& OutPathToEndPlugin);
+
 private:
 	FPluginReferenceViewerUtils() = delete;
 	FPluginReferenceViewerUtils(const FPluginReferenceViewerUtils&) = delete;
diff --git a/Engine/Source/Programs/UnrealBuildTool/System/PluginDescriptor.cs b/Engine/Source/Programs/UnrealBuildTool/System/PluginDescriptor.cs
--- a/Engine/Source/Programs/UnrealBuildTool/System/PluginDescriptor.cs
+++ b/Engine/Source/Programs/UnrealBuildTool/System/PluginDescriptor.cs
@@ -371,7 +371,16 @@ namespace UnrealBuildTool
 				Plugins = Array.ConvertAll(PluginsArray, x => PluginReferenceDescriptor.FromJsonObject(x)).ToList();
 			}
 
-			RawObject.TryGetStringArrayField("DisallowedPlugins", out DisallowedPlugins);
+			JsonObject[]? DisallowedPluginsArray;
+			if (RawObject.TryGetObjectArrayField("DisallowedPlugins", out DisallowedPluginsArray))
+			{
+				DisallowedPlugins = Array.ConvertAll(DisallowedPluginsArray, x => x.GetStringField("Name"));
+			}
+			else
+			{
+				// Backwards compatibility still check for a simple array.
+				RawObject.TryGetStringArrayField("DisallowedPlugins", out DisallowedPlugins);
+			}
 		}
 
 		/// <summary>
