Fix crash in plugin reference viewer when a plugin has an optional dependency on a plugin that doesn't exist on disk

diff --git a/Engine/Plugins/Experimental/PluginReferenceViewer/Source/PluginReferenceViewer/Private/EdGraph_PluginReferenceViewer.cpp b/Engine/Plugins/Experimental/PluginReferenceViewer/Source/PluginReferenceViewer/Private/EdGraph_PluginReferenceViewer.cpp
--- a/Engine/Plugins/Experimental/PluginReferenceViewer/Source/PluginReferenceViewer/Private/EdGraph_PluginReferenceViewer.cpp
+++ b/Engine/Plugins/Experimental/PluginReferenceViewer/Source/PluginReferenceViewer/Private/EdGraph_PluginReferenceViewer.cpp
@@ -2,6 +2,7 @@
 
 #include "EdGraph_PluginReferenceViewer.h"
 
+#include "Algo/Find.h"
 #include "AssetManagerEditorModule.h"
 #include "AssetRegistry/AssetRegistryModule.h"
 #include "AssetRegistry/AssetData.h"
@@ -26,8 +27,13 @@ namespace UE::PluginReferenceViewer::Private
 		return false;
 	}
 
-	void FindPluginAssetReferences(const TSharedRef<IPlugin>& ParentPlugin, const TSharedRef<IPlugin>& ChildPlugin, TArray<FAssetIdentifier>& OutAssetIdentifiers)
+	void FindPluginAssetReferences(const TSharedPtr<IPlugin>& ParentPlugin, const TSharedPtr<IPlugin>& ChildPlugin, TArray<FAssetIdentifier>& OutAssetIdentifiers)
 	{
+		if (!ParentPlugin.IsValid() || !ChildPlugin.IsValid())
+		{
+			return;
+		}
+
 		IAssetRegistry& AssetRegistry = FModuleManager::LoadModuleChecked<FAssetRegistryModule>("AssetRegistry").Get();
 
 		const FString ChildMountedAssetPath = ChildPlugin->GetMountedAssetPath();
@@ -63,8 +69,13 @@ namespace UE::PluginReferenceViewer::Private
 		}
 	};
 
-	void GetPluginAssets(const TSharedRef<IPlugin>& Plugin, TArray<FAssetIdentifier>& OutAssetIdentifiers)
+	void GetPluginAssets(const TSharedPtr<IPlugin>& Plugin, TArray<FAssetIdentifier>& OutAssetIdentifiers)
 	{
+		if (!Plugin.IsValid())
+		{
+			return;
+		}
+
 		IAssetRegistry& AssetRegistry = FModuleManager::LoadModuleChecked<FAssetRegistryModule>("AssetRegistry").Get();
 
 		FARFilter Filter;
@@ -82,8 +93,13 @@ namespace UE::PluginReferenceViewer::Private
 		}
 	};
 
-	uint64 GetSizeOfPlugin(const TSharedRef<IPlugin> Plugin)
+	uint64 GetSizeOfPlugin(const TSharedPtr<IPlugin> Plugin)
 	{
+		if (!Plugin.IsValid())
+		{
+			return 0;
+		}
+
 		uint64 Size = 0;
 		TArray<FAssetData> PluginAssets;
 
@@ -144,15 +160,15 @@ const TArray<FPluginIdentifier>& UEdGraph_PluginReferenceViewer::GetCurrentGraph
 
 void UEdGraph_PluginReferenceViewer::FindAssetReferencesAcrossPlugins(const FPluginIdentifier& From, const FPluginIdentifier& To, TArray<FAssetIdentifier>& OutAssetIdentifiers)
 {
-	const TSharedRef<IPlugin> FromPlugin = PluginMap.FindChecked(From);
-	const TSharedRef<IPlugin> ToPlugin = PluginMap.FindChecked(To);
+	const TSharedPtr<IPlugin> FromPlugin = PluginMap.FindRef(From);
+	const TSharedPtr<IPlugin> ToPlugin = PluginMap.FindRef(To);
 
 	UE::PluginReferenceViewer::Private::FindPluginAssetReferences(FromPlugin, ToPlugin, OutAssetIdentifiers);
 }
 
 void UEdGraph_PluginReferenceViewer::GetPluginAssets(const FPluginIdentifier& Plugin, TArray<FAssetIdentifier>& OutAssetIdentifiers)
 {
-	const TSharedRef<IPlugin> ForPlugin = PluginMap.FindChecked(Plugin);
+	const TSharedPtr<IPlugin> ForPlugin = PluginMap.FindRef(Plugin);
 
 	UE::PluginReferenceViewer::Private::GetPluginAssets(ForPlugin, OutAssetIdentifiers);
 }
@@ -186,7 +202,7 @@ void UEdGraph_PluginReferenceViewer::CachePluginDependencies(const TArray<TShare
 			for (const TSharedRef<IPlugin>& Plugin : Plugins)
 			{
 				FPluginIdentifier NewPluginIdentifier = FPluginIdentifier::FromString(Plugin->GetName());
-				Nodes.Add(NewPluginIdentifier, MakeUnique<FPluginDependsNode>(NewPluginIdentifier));
+				Nodes.Add(NewPluginIdentifier, MakeUnique<FPluginDependsNode>(NewPluginIdentifier, /* bMissing */ false));
 			}
 
 			for (const TSharedRef<IPlugin>& Plugin : Plugins)
@@ -205,7 +221,7 @@ void UEdGraph_PluginReferenceViewer::CachePluginDependencies(const TArray<TShare
 					else
 					{
 						FPluginIdentifier NewPluginIdentifier = FPluginIdentifier::FromString(Item.Name);
-						ChildNode = Nodes.Add(NewPluginIdentifier, MakeUnique<FPluginDependsNode>(NewPluginIdentifier)).Get();
+						ChildNode = Nodes.Add(NewPluginIdentifier, MakeUnique<FPluginDependsNode>(NewPluginIdentifier, /* bMissing */ true)).Get();
 					}
 
 					ParentNode->Dependencies.Add(ChildNode);
@@ -234,7 +250,7 @@ UEdGraphNode_PluginReference* UEdGraph_PluginReferenceViewer::ConstructNodes(con
 			RootNodeInfo.Parents.Emplace(NAME_None);
 		}
 
-		RecursivelyPopulateNodeInfos(true, GraphRootIdentifiers, NewReferenceNodeInfos, 0, SettingsReferencerMaxDepth);
+		RecursivelyPopulateNodeInfos(/* bReferencers */ true, GraphRootIdentifiers, NewReferenceNodeInfos, 0, SettingsReferencerMaxDepth);
 
 		TMap<FPluginIdentifier, FPluginReferenceNodeInfo> NewDependencyNodeInfos;
 		for (const FPluginIdentifier& RootIdentifier : GraphRootIdentifiers)
@@ -243,22 +259,7 @@ UEdGraphNode_PluginReference* UEdGraph_PluginReferenceViewer::ConstructNodes(con
 			DRootNodeInfo.Parents.Emplace(NAME_None);
 		}
 
-		RecursivelyPopulateNodeInfos(false, GraphRootIdentifiers, NewDependencyNodeInfos, 0, SettingsDependencyMaxDepth);
-
-		// Store the Plugin in the NodeInfos
-		for (TPair<FPluginIdentifier, FPluginReferenceNodeInfo>& InfoPair : NewReferenceNodeInfos)
-		{
-			InfoPair.Value.Plugin = PluginMap.FindChecked(InfoPair.Key);
-		}
-
-		for (TPair<FPluginIdentifier, FPluginReferenceNodeInfo>& InfoPair : NewDependencyNodeInfos)
-		{
-			const TSharedRef<IPlugin>* FoundPlugin = PluginMap.Find(InfoPair.Key);
-			if (FoundPlugin != nullptr)
-			{
-				InfoPair.Value.Plugin = *FoundPlugin;
-			}
-		}
+		RecursivelyPopulateNodeInfos(/* bReferencers */ false, GraphRootIdentifiers, NewDependencyNodeInfos, 0, SettingsDependencyMaxDepth);
 
 		ReferencerNodeInfos = NewReferenceNodeInfos;
 		DependencyNodeInfos = NewDependencyNodeInfos;
@@ -295,7 +296,7 @@ void UEdGraph_PluginReferenceViewer::LoadAdvancedPluginInfo()
 		[this, &MapKeys](int32 Index)
 		{
 			FPluginStats& Stats = StatsMap.FindChecked(MapKeys[Index]);
-			Stats.Size = UE::PluginReferenceViewer::Private::GetSizeOfPlugin(PluginMap.FindChecked(MapKeys[Index]));
+			Stats.Size = UE::PluginReferenceViewer::Private::GetSizeOfPlugin(PluginMap.FindRef(MapKeys[Index]));
 		});
 
 	MapKeys.Empty();
@@ -305,7 +306,7 @@ void UEdGraph_PluginReferenceViewer::LoadAdvancedPluginInfo()
 		[this, &MapKeys](int32 Index)
 		{
 			FPluginStats& Stats = StatsMap.FindChecked(MapKeys[Index]);
-			Stats.Size = UE::PluginReferenceViewer::Private::GetSizeOfPlugin(PluginMap.FindChecked(MapKeys[Index]));
+			Stats.Size = UE::PluginReferenceViewer::Private::GetSizeOfPlugin(PluginMap.FindRef(MapKeys[Index]));
 		});
 
 
@@ -348,7 +349,7 @@ UEdGraphNode_PluginReference* UEdGraph_PluginReferenceViewer::RecursivelyCreateN
 		const bool bIsADuplicate = NodeInfo.Parents.Num() > 1;
 
 		NewNode = CreatePluginReferenceNode();
-		NewNode->SetupPluginReferenceNode(InNodeLoc, NodeInfo.Identifier, NodeInfo.Plugin, !bIsCompactMode, bIsADuplicate);
+		NewNode->SetupPluginReferenceNode(InNodeLoc, NodeInfo.Identifier, PluginMap.FindRef(NodeInfo.Identifier), !bIsCompactMode, bIsADuplicate);
 		NodeProvSize = NodeInfo.ProvisionSize(InParentId);
 	}
 
@@ -508,7 +509,7 @@ UEdGraphNode_PluginReference* UEdGraph_PluginReferenceViewer::RefilterGraph()
 
 		const FPluginReferenceNodeInfo& NodeInfo = ReferencerNodeInfos[FirstGraphRootIdentifier];
 		RootNode = CreatePluginReferenceNode();
-		RootNode->SetupPluginReferenceNode(CurrentGraphRootOrigin, NodeInfo.Identifier, NodeInfo.Plugin, !bIsCompactMode, false);
+		RootNode->SetupPluginReferenceNode(CurrentGraphRootOrigin, NodeInfo.Identifier, PluginMap.FindRef(NodeInfo.Identifier), !bIsCompactMode, false);
 
 		// Show referencers
 		RecursivelyFilterNodeInfos(FirstGraphRootIdentifier, ReferencerNodeInfos, 0, SettingsReferencerMaxDepth);
@@ -592,16 +593,20 @@ void UEdGraph_PluginReferenceViewer::GetPluginDependencies(const FPluginIdentifi
 	const bool bShowOptionalPlugins = PluginReferenceViewer.Pin()->IsShowOptionalPluginsChecked();
 
 	const FPluginDependsNode* PluginNode = CachedDependsNodes.FindChecked(PluginIdentifier).Get();
-	const TSharedRef<IPlugin> Plugin = PluginMap.FindChecked(PluginNode->Identifier);
+	const TSharedPtr<IPlugin> Plugin = PluginMap.FindRef(PluginNode->Identifier);
+	if (!Plugin.IsValid()) // Missing optional plugin
+	{
+		return;
+	}
 
 	for (int32 i = 0; i < PluginNode->Dependencies.Num(); ++i)
 	{
 		const FPluginDependsNode* ChildNode = PluginNode->Dependencies[i];
 
 		// Plugin dependencies may not be discovered
-		if (const TSharedRef<IPlugin>* ChildPlugin = PluginMap.Find(ChildNode->Identifier))
+		if (const TSharedPtr<IPlugin> ChildPlugin = PluginMap.FindRef(ChildNode->Identifier))
 		{
-			if ((*ChildPlugin)->GetType() == EPluginType::Engine && !bShowEnginePlugins)
+			if (ChildPlugin->GetType() == EPluginType::Engine && !bShowEnginePlugins)
 			{
 				continue;
 			}
@@ -624,20 +629,22 @@ void UEdGraph_PluginReferenceViewer::GetPluginReferencers(const FPluginIdentifie
 	const FPluginDependsNode* PluginNode = CachedDependsNodes.FindChecked(PluginIdentifier).Get();
 	for (const FPluginDependsNode* Referencer : PluginNode->Referencers)
 	{
-		for (int32 i = 0; i < Referencer->Dependencies.Num(); ++i)
+		const TSharedPtr<IPlugin> ReferencerPlugin = PluginMap.FindRef(Referencer->Identifier);
+		if (!ReferencerPlugin.IsValid()) // Should not happen as we only make plugin A -> plugin B references when looking at plugin A in the plugin manager 
 		{
-			const FPluginDependsNode* ParentNode = Referencer;
+			continue;
+		}
 
-			if (ParentNode->Dependencies[i] == PluginNode)
-			{
-				const TSharedRef<IPlugin> ParentPlugin = PluginMap.FindChecked(Referencer->Identifier);
-				const FPluginReferenceDescriptor& ReferenceDescriptor = ParentPlugin->GetDescriptor().Plugins[i];
-				check(ReferenceDescriptor.Name == PluginNode->Identifier.Name.ToString());
+		FPluginDependsNode* const* FoundDependency = Algo::Find(Referencer->Dependencies, PluginNode);
+		checkf(FoundDependency, TEXT("Referencer mapping from %s to %s is not symmetric"), *Referencer->Identifier.ToString(), *PluginIdentifier.ToString());
+		
+		// Indices are assumed to match from iteration order in CachePluginDependencies
+		int32 Index = FoundDependency - Referencer->Dependencies.GetData();
 
-				OutReferencers.Add(FPluginDependency(ParentNode->Identifier, ReferenceDescriptor.bEnabled, ReferenceDescriptor.bOptional));
-				break;
-			}
-		}
+		const FPluginReferenceDescriptor& ReferenceDescriptor = ReferencerPlugin->GetDescriptor().Plugins[Index];
+		check(ReferenceDescriptor.Name == PluginNode->Identifier.Name.ToString());
+
+		OutReferencers.Add(FPluginDependency(Referencer->Identifier, ReferenceDescriptor.bEnabled, ReferenceDescriptor.bOptional));
 	}
 }
 
diff --git a/Engine/Plugins/Experimental/PluginReferenceViewer/Source/PluginReferenceViewer/Private/EdGraph_PluginReferenceViewer.h b/Engine/Plugins/Experimental/PluginReferenceViewer/Source/PluginReferenceViewer/Private/EdGraph_PluginReferenceViewer.h
--- a/Engine/Plugins/Experimental/PluginReferenceViewer/Source/PluginReferenceViewer/Private/EdGraph_PluginReferenceViewer.h
+++ b/Engine/Plugins/Experimental/PluginReferenceViewer/Source/PluginReferenceViewer/Private/EdGraph_PluginReferenceViewer.h
@@ -75,22 +75,22 @@ struct FPluginStats
 
 struct FPluginDependsNode
 {
-	explicit FPluginDependsNode(const FPluginIdentifier& InIdentifier)
+	explicit FPluginDependsNode(const FPluginIdentifier& InIdentifier, bool InMissing)
 		: Identifier(InIdentifier)
+		, bMissing(InMissing)
 	{
 	}
 
 	FPluginIdentifier Identifier;
-	TArray<FPluginDependsNode*> Dependencies;
-	TArray<FPluginDependsNode*> Referencers;
+	TArray<FPluginDependsNode*> Dependencies; // Raw pointer to another node in CachedDependsNodes, same lifetime
+	TArray<FPluginDependsNode*> Referencers; // Raw pointer to another node in CachedDependsNodes, same lifetime
+	bool bMissing = false;
 };
 
 struct FPluginReferenceNodeInfo
 {
 	FPluginIdentifier Identifier;
 
-	TSharedPtr<IPlugin> Plugin;
-
 	TArray<TPair<FPluginIdentifier, EPluginReferencePinCategory>> Children;
 
 	TArray<FPluginIdentifier> Parents;
@@ -187,11 +187,15 @@ private:
 	TArray<FPluginIdentifier> CurrentGraphRootIdentifiers;
 	FIntPoint CurrentGraphRootOrigin;
 
+	// Mapping of plugins which were seen when traversing the dependency graph of active plugins, whether or not they exist in the plugin manager 
 	TMap<FPluginIdentifier, TUniquePtr<FPluginDependsNode>> CachedDependsNodes;
-	TMap<FPluginIdentifier, const TSharedRef<IPlugin>> PluginMap;
+	// Mapping of plugins which actually exist - may have fewer keys than CachedDependsNodes
+	TMap<FPluginIdentifier, const TSharedPtr<IPlugin>> PluginMap;
 	TMap<FPluginIdentifier, FPluginStats> StatsMap;
 
+	// Current set of referencers of the graph root plugin(s) to be displayed, expanded to configured depth
 	TMap<FPluginIdentifier, FPluginReferenceNodeInfo> ReferencerNodeInfos;
+	// Current set of dependencies of the graph root plugin(s) to be displayed, expanded to configured depth 
 	TMap<FPluginIdentifier, FPluginReferenceNodeInfo> DependencyNodeInfos;
 
 	friend SPluginReferenceViewer;
