Add "Show in Content Browser" option for plugins in the Plugin Reference Viewer.

diff --git a/Engine/Plugins/Experimental/PluginReferenceViewer/Source/PluginReferenceViewer/PluginReferenceViewer.Build.cs b/Engine/Plugins/Experimental/PluginReferenceViewer/Source/PluginReferenceViewer/PluginReferenceViewer.Build.cs
--- a/Engine/Plugins/Experimental/PluginReferenceViewer/Source/PluginReferenceViewer/PluginReferenceViewer.Build.cs
+++ b/Engine/Plugins/Experimental/PluginReferenceViewer/Source/PluginReferenceViewer/PluginReferenceViewer.Build.cs
@@ -38,6 +38,7 @@ public class PluginReferenceViewer : ModuleRules
 				"ToolMenus",
 				"ToolWidgets",
 				"PluginBrowser",
+				"ContentBrowserData"
 			}
 		);
 	}
diff --git a/Engine/Plugins/Experimental/PluginReferenceViewer/Source/PluginReferenceViewer/Private/PluginReferenceViewerCommands.cpp b/Engine/Plugins/Experimental/PluginReferenceViewer/Source/PluginReferenceViewer/Private/PluginReferenceViewerCommands.cpp
--- a/Engine/Plugins/Experimental/PluginReferenceViewer/Source/PluginReferenceViewer/Private/PluginReferenceViewerCommands.cpp
+++ b/Engine/Plugins/Experimental/PluginReferenceViewer/Source/PluginReferenceViewer/Private/PluginReferenceViewerCommands.cpp
@@ -21,6 +21,7 @@ void FPluginReferenceViewerCommands::RegisterCommands()
 	UI_COMMAND(ShowOptionalPlugins, "Show Optional Plugins", "Toggles visibility of Optional Plugin dependencies", EUserInterfaceActionType::ToggleButton, FInputChord());
 
 	UI_COMMAND(OpenPluginProperties, "Open Plugin Properties", "Opens the Plugin Properties Editor Window", EUserInterfaceActionType::Button, FInputChord());
+	UI_COMMAND(ShowInContentBrowser, "Show in Content Browser", "Navigates to the Plugin in the Content Browser", EUserInterfaceActionType::Button, FInputChord());
 	UI_COMMAND(ZoomToFit, "Zoom to Fit", "Zoom in and center the view on the selected item", EUserInterfaceActionType::Button, FInputChord(EKeys::F));
 	UI_COMMAND(ReCenterGraph, "Re-Center Graph", "Re-centers the graph on this node, showing all referencers and references for this Plugin instead", EUserInterfaceActionType::Button, FInputChord());
 }
diff --git a/Engine/Plugins/Experimental/PluginReferenceViewer/Source/PluginReferenceViewer/Private/PluginReferenceViewerCommands.h b/Engine/Plugins/Experimental/PluginReferenceViewer/Source/PluginReferenceViewer/Private/PluginReferenceViewerCommands.h
--- a/Engine/Plugins/Experimental/PluginReferenceViewer/Source/PluginReferenceViewer/Private/PluginReferenceViewerCommands.h
+++ b/Engine/Plugins/Experimental/PluginReferenceViewer/Source/PluginReferenceViewer/Private/PluginReferenceViewerCommands.h
@@ -20,6 +20,7 @@ public:
 	TSharedPtr<FUICommandInfo> ShowOptionalPlugins;
 
 	TSharedPtr<FUICommandInfo> OpenPluginProperties;
+	TSharedPtr<FUICommandInfo> ShowInContentBrowser;
 	TSharedPtr<FUICommandInfo> ZoomToFit;
 	TSharedPtr<FUICommandInfo> ReCenterGraph;
 };
\ No newline at end of file
diff --git a/Engine/Plugins/Experimental/PluginReferenceViewer/Source/PluginReferenceViewer/Private/PluginReferenceViewerSchema.cpp b/Engine/Plugins/Experimental/PluginReferenceViewer/Source/PluginReferenceViewer/Private/PluginReferenceViewerSchema.cpp
--- a/Engine/Plugins/Experimental/PluginReferenceViewer/Source/PluginReferenceViewer/Private/PluginReferenceViewerSchema.cpp
+++ b/Engine/Plugins/Experimental/PluginReferenceViewer/Source/PluginReferenceViewer/Private/PluginReferenceViewerSchema.cpp
@@ -58,6 +58,7 @@ void UPluginReferenceViewerSchema::GetContextMenuActions(UToolMenu* Menu, UGraph
 	{
 		FToolMenuSection& Section = Menu->AddSection(TEXT("Plugin"), NSLOCTEXT("PluginReferenceViewerSchema", "PluginSectionLabel", "Plugin"));
 		Section.AddMenuEntry(FPluginReferenceViewerCommands::Get().OpenPluginProperties);
+		Section.AddMenuEntry(FPluginReferenceViewerCommands::Get().ShowInContentBrowser);
 	}
 
 	{
diff --git a/Engine/Plugins/Experimental/PluginReferenceViewer/Source/PluginReferenceViewer/Private/SPluginReferenceViewer.cpp b/Engine/Plugins/Experimental/PluginReferenceViewer/Source/PluginReferenceViewer/Private/SPluginReferenceViewer.cpp
--- a/Engine/Plugins/Experimental/PluginReferenceViewer/Source/PluginReferenceViewer/Private/SPluginReferenceViewer.cpp
+++ b/Engine/Plugins/Experimental/PluginReferenceViewer/Source/PluginReferenceViewer/Private/SPluginReferenceViewer.cpp
@@ -2,8 +2,11 @@
 
 #include "SPluginReferenceViewer.h"
 
+#include "ContentBrowserModule.h"
 #include "EdGraphNode_PluginReference.h"
 #include "EdGraph_PluginReferenceViewer.h"
+#include "IContentBrowserDataModule.h"
+#include "IContentBrowserSingleton.h"
 #include "Features/EditorFeatures.h"
 #include "Features/IModularFeatures.h"
 #include "Features/IPluginsEditorFeature.h"
@@ -183,6 +186,53 @@ void SPluginReferenceViewer::OnOpenPluginProperties()
 	}
 }
 
+void SPluginReferenceViewer::OnShowPluginInContentBrowser()
+{
+	TSharedPtr<const IPlugin> Plugin;
+
+	const FGraphPanelSelectionSet& SelectedNodes = GraphEditorPtr->GetSelectedNodes();
+	if (!SelectedNodes.IsEmpty())
+	{
+		for (UObject* Node : SelectedNodes)
+		{
+			UEdGraphNode_PluginReference* PluginReferenceNode = Cast<UEdGraphNode_PluginReference>(Node);
+			if (PluginReferenceNode)
+			{
+				Plugin = PluginReferenceNode->GetPlugin();
+				break;
+			}
+		}
+	}
+
+	if (Plugin != nullptr)
+	{
+		UContentBrowserDataSubsystem* ContentBrowserData = IContentBrowserDataModule::Get().GetSubsystem();
+
+		FString MountedPath = Plugin->GetMountedAssetPath();
+		FStringView MountedPathView(MountedPath);
+		
+		// Remove any trailing slashes
+		if (MountedPathView.EndsWith('/'))
+		{
+			MountedPathView.LeftChopInline(1);
+		}
+
+		FName VirtualPath;
+		ContentBrowserData->ConvertInternalPathToVirtual(MountedPathView, VirtualPath);
+
+		if (!VirtualPath.IsNone())
+		{
+			FContentBrowserItem Item = ContentBrowserData->GetItemAtPath(VirtualPath, EContentBrowserItemTypeFilter::IncludeAll);
+			if (Item.IsValid())
+			{
+				TArray<FContentBrowserItem> Items = { Item };
+				IContentBrowserSingleton& ContentBrowser = FModuleManager::Get().GetModuleChecked<FContentBrowserModule>("ContentBrowser").Get();
+				ContentBrowser.SyncBrowserToItems(Items);
+			}	
+		}
+	}
+}
+
 void SPluginReferenceViewer::OnApplyHistoryData(const FPluginReferenceViewerHistoryData& History)
 {
 	if (GraphObj)
@@ -502,6 +552,11 @@ void SPluginReferenceViewer::RegisterActions()
 		FPluginReferenceViewerCommands::Get().OpenPluginProperties,
 		FExecuteAction::CreateSP(this, &SPluginReferenceViewer::OnOpenPluginProperties),
 		FCanExecuteAction::CreateSP(this, &SPluginReferenceViewer::HasAtLeastOneRealNodeSelected));
+	
+	PluginReferenceViewerActions->MapAction(
+		FPluginReferenceViewerCommands::Get().ShowInContentBrowser,
+		FExecuteAction::CreateSP(this, &SPluginReferenceViewer::OnShowPluginInContentBrowser),
+		FCanExecuteAction::CreateSP(this, &SPluginReferenceViewer::HasAtLeastOneRealNodeSelected));
 
 	PluginReferenceViewerActions->MapAction(
 		FPluginReferenceViewerCommands::Get().ZoomToFit,
diff --git a/Engine/Plugins/Experimental/PluginReferenceViewer/Source/PluginReferenceViewer/Private/SPluginReferenceViewer.h b/Engine/Plugins/Experimental/PluginReferenceViewer/Source/PluginReferenceViewer/Private/SPluginReferenceViewer.h
--- a/Engine/Plugins/Experimental/PluginReferenceViewer/Source/PluginReferenceViewer/Private/SPluginReferenceViewer.h
+++ b/Engine/Plugins/Experimental/PluginReferenceViewer/Source/PluginReferenceViewer/Private/SPluginReferenceViewer.h
@@ -45,6 +45,7 @@ public:
 
 private:
 	void OnOpenPluginProperties();
+	void OnShowPluginInContentBrowser();
 	bool HasAtLeastOneRealNodeSelected();
 
 	void OpenPluginProperties(const FString& PluginName);
