Normalize the "Show in Content Browser" language for the plugin reference viewer to match the "Browse To XXX" language of other places in the editor.
- Add shortcut.
- Add icon.

diff --git a/Engine/Plugins/Experimental/PluginReferenceViewer/Source/PluginReferenceViewer/Private/PluginReferenceViewerCommands.cpp b/Engine/Plugins/Experimental/PluginReferenceViewer/Source/PluginReferenceViewer/Private/PluginReferenceViewerCommands.cpp
--- a/Engine/Plugins/Experimental/PluginReferenceViewer/Source/PluginReferenceViewer/Private/PluginReferenceViewerCommands.cpp
+++ b/Engine/Plugins/Experimental/PluginReferenceViewer/Source/PluginReferenceViewer/Private/PluginReferenceViewerCommands.cpp
@@ -21,7 +21,7 @@ void FPluginReferenceViewerCommands::RegisterCommands()
 	UI_COMMAND(ShowOptionalPlugins, "Show Optional Plugins", "Toggles visibility of Optional Plugin dependencies", EUserInterfaceActionType::ToggleButton, FInputChord());
 
 	UI_COMMAND(OpenPluginProperties, "Open Plugin Properties", "Opens the Plugin Properties Editor Window", EUserInterfaceActionType::Button, FInputChord());
-	UI_COMMAND(ShowInContentBrowser, "Show in Content Browser", "Navigates to the Plugin in the Content Browser", EUserInterfaceActionType::Button, FInputChord());
+	UI_COMMAND(ShowInContentBrowser, "Browse to Plugin", "Browses to the associated Plugin and selects it in the most recently used Content Browser (summoning one if necessary)", EUserInterfaceActionType::Button, FInputChord(EModifierKey::Control, EKeys::B));
 	UI_COMMAND(ZoomToFit, "Zoom to Fit", "Zoom in and center the view on the selected item", EUserInterfaceActionType::Button, FInputChord(EKeys::F));
 	UI_COMMAND(ReCenterGraph, "Re-Center Graph", "Re-centers the graph on this node, showing all referencers and references for this Plugin instead", EUserInterfaceActionType::Button, FInputChord());
 }
diff --git a/Engine/Plugins/Experimental/PluginReferenceViewer/Source/PluginReferenceViewer/Private/PluginReferenceViewerSchema.cpp b/Engine/Plugins/Experimental/PluginReferenceViewer/Source/PluginReferenceViewer/Private/PluginReferenceViewerSchema.cpp
--- a/Engine/Plugins/Experimental/PluginReferenceViewer/Source/PluginReferenceViewer/Private/PluginReferenceViewerSchema.cpp
+++ b/Engine/Plugins/Experimental/PluginReferenceViewer/Source/PluginReferenceViewer/Private/PluginReferenceViewerSchema.cpp
@@ -57,8 +57,11 @@ void UPluginReferenceViewerSchema::GetContextMenuActions(UToolMenu* Menu, UGraph
 
 	{
 		FToolMenuSection& Section = Menu->AddSection(TEXT("Plugin"), NSLOCTEXT("PluginReferenceViewerSchema", "PluginSectionLabel", "Plugin"));
+		
 		Section.AddMenuEntry(FPluginReferenceViewerCommands::Get().OpenPluginProperties);
-		Section.AddMenuEntry(FPluginReferenceViewerCommands::Get().ShowInContentBrowser);
+		
+		FToolMenuEntry& ShowInEntry = Section.AddMenuEntry(FPluginReferenceViewerCommands::Get().ShowInContentBrowser);
+		ShowInEntry.Icon = FSlateIcon(FAppStyle::GetAppStyleSetName(), "SystemWideCommands.FindInContentBrowser.Small");
 	}
 
 	{
