Exposing plugin Reference Viewer Utility functions to use within Validators. Expanding validator to ensure plugins that use Gameplaytags from another plugin have dependency on that plugin.

diff --git a/Engine/Plugins/Experimental/PluginReferenceViewer/Source/PluginReferenceViewer/Private/PluginReferenceViewerUtils.cpp b/Engine/Plugins/Experimental/PluginReferenceViewer/Source/PluginReferenceViewer/Private/PluginReferenceViewerUtils.cpp
--- a/Engine/Plugins/Experimental/PluginReferenceViewer/Source/PluginReferenceViewer/Private/PluginReferenceViewerUtils.cpp
+++ b/Engine/Plugins/Experimental/PluginReferenceViewer/Source/PluginReferenceViewer/Private/PluginReferenceViewerUtils.cpp
@@ -196,130 +196,6 @@ namespace PluginReferenceViewerUtils
 		FString tmp;
 		TracePathFromPluginToPlugin(StartPluginName, EndPluginName, tmp);
 	}
-
-	TArray<FAssetIdentifier> GetAssetDependencies(const TSharedRef<IPlugin>& InPlugin)
-	{
-		IAssetRegistry& AssetRegistry = FModuleManager::LoadModuleChecked<FAssetRegistryModule>("AssetRegistry").Get();
-
-		FARFilter Filter;
-		Filter.bRecursivePaths = true;
-		Filter.PackagePaths.Add(FName(InPlugin->GetMountedAssetPath()));
-
-		TArray<FAssetData> AssetsInPlugin;
-		AssetRegistry.GetAssets(Filter, AssetsInPlugin);
-
-		TSet<FAssetIdentifier> UniqueDependencies;
-
-		for (const FAssetData& AssetData : AssetsInPlugin)
-		{
-			TArray<FAssetIdentifier> AssetReferencers;
-			AssetRegistry.GetDependencies(AssetData.PackageName, AssetReferencers);
-
-			for (const FAssetIdentifier& AssetReferencer : AssetReferencers)
-			{
-				UniqueDependencies.Add(AssetReferencer);
-			}
-		}
-
-		return UniqueDependencies.Array();
-	}
-
-	TMap<FString, TArray<FAssetIdentifier>> SplitByPlugins(const TSharedRef<IPlugin>& InOwningPlugin, const TArray<FAssetIdentifier>& InPluginDependencies)
-	{
-		TMap<FString, TArray<FAssetIdentifier>> Results;
-
-		const FName GameplayTagStructPackage = FGameplayTag::StaticStruct()->GetOutermost()->GetFName();
-		const FName NAME_GameplayTag = FGameplayTag::StaticStruct()->GetFName();
-
-		for (const FAssetIdentifier& AssetIdentifier : InPluginDependencies)
-		{
-			const FString PackageNameString = AssetIdentifier.PackageName.ToString();
-
-			if (AssetIdentifier.ObjectName == NAME_GameplayTag && AssetIdentifier.PackageName == GameplayTagStructPackage)
-			{
-				const TArray<TSharedRef<IPlugin>> SourcePlugins = FPluginReferenceViewerUtils::FindGameplayTagSourcePlugins(AssetIdentifier.ValueName);
-				for (const TSharedRef<IPlugin>& SourcePlugin : SourcePlugins)
-				{
-					if (SourcePlugin != InOwningPlugin)
-					{
-						Results.FindOrAdd(SourcePlugin->GetName()).Add(AssetIdentifier);
-					}
-				}
-			}
-			else if (FStringView ModuleName; FPackageName::TryConvertScriptPackageNameToModuleName(PackageNameString, ModuleName))
-			{
-				const TSharedPtr<IPlugin> ModulePlugin = IPluginManager::Get().GetModuleOwnerPlugin(FName(ModuleName));
-				if (ModulePlugin.IsValid() && ModulePlugin != InOwningPlugin)
-				{
-					Results.FindOrAdd(ModulePlugin->GetName()).Add(AssetIdentifier);
-				}
-			}
-			else
-			{
-				const TSharedPtr<IPlugin> PackagePlugin = IPluginManager::Get().FindPluginFromPath(PackageNameString);
-				if (PackagePlugin.IsValid() && PackagePlugin != InOwningPlugin)
-				{
-					Results.FindOrAdd(PackagePlugin->GetName()).Add(AssetIdentifier);
-				}
-			}
-		}
-
-		return Results;
-	}
-
-	void SplitByReferenceType(const TArray<FAssetIdentifier>& InAssetIdentifiers, TArray<FAssetIdentifier>& OutAssetReferences, TArray<FAssetIdentifier>& OutScriptReferences, TArray<FAssetIdentifier>& OutNameReferences)
-	{
-		const FName GameplayTagStructPackage = FGameplayTag::StaticStruct()->GetOutermost()->GetFName();
-		const FName NAME_GameplayTag = FGameplayTag::StaticStruct()->GetFName();
-		const FName NAME_DataTable(TEXT("DataTable"));
-
-		TArray<FName> AllPackageNames;
-		for (const FAssetIdentifier& AssetIdentifier : InAssetIdentifiers)
-		{
-			AllPackageNames.Add(AssetIdentifier.PackageName);
-		}
-
-		TMap<FName, FAssetData> PackagesToAssetDataMap;
-		UE::AssetRegistry::GetAssetForPackages(AllPackageNames, PackagesToAssetDataMap);
-
-		for (const FAssetIdentifier& AssetIdentifier : InAssetIdentifiers)
-		{
-			const FString PackageNameString = AssetIdentifier.PackageName.ToString();
-
-			if (AssetIdentifier.ObjectName == NAME_GameplayTag && AssetIdentifier.PackageName == GameplayTagStructPackage)
-			{
-				OutNameReferences.Add(AssetIdentifier);
-			}
-			else if (FPackageName::IsScriptPackage(PackageNameString))
-			{
-				FStringView ModuleName;
-				if (FPackageName::TryConvertScriptPackageNameToModuleName(PackageNameString, ModuleName))
-				{
-					OutScriptReferences.Add(AssetIdentifier);
-				}
-			}
-			else
-			{
-				// PackagesToAssetDataMap may not contain the package as references to assets that no longer exist do occur.
-				const FAssetData* AssetData = PackagesToAssetDataMap.Find(AssetIdentifier.PackageName);
-				if (AssetData != nullptr)
-				{
-					if (AssetData->AssetClassPath.GetAssetName() == NAME_DataTable)
-					{
-						OutNameReferences.Add(AssetIdentifier);
-					}
-					else
-					{
-						OutAssetReferences.Add(AssetIdentifier);
-					}
-				}
-				else
-				{
-					UE_LOG(LogPluginReferenceViewerUtils, Display, TEXT("Skipping package '%s' due to missing asset data. Package may no longer exist!"), *AssetIdentifier.PackageName.ToString());
-				}
-			}
-		}
-	}
 }
 
 namespace PluginReferenceViewerCVars
@@ -362,6 +238,130 @@ namespace PluginReferenceViewerCVars
 	);
 }
 
+/*static*/TArray<FAssetIdentifier> FPluginReferenceViewerUtils::GetAssetDependencies(const TSharedRef<IPlugin>& InPlugin)
+{
+	IAssetRegistry& AssetRegistry = FModuleManager::LoadModuleChecked<FAssetRegistryModule>("AssetRegistry").Get();
+
+	FARFilter Filter;
+	Filter.bRecursivePaths = true;
+	Filter.PackagePaths.Add(FName(InPlugin->GetMountedAssetPath()));
+
+	TArray<FAssetData> AssetsInPlugin;
+	AssetRegistry.GetAssets(Filter, AssetsInPlugin);
+
+	TSet<FAssetIdentifier> UniqueDependencies;
+
+	for (const FAssetData& AssetData : AssetsInPlugin)
+	{
+		TArray<FAssetIdentifier> AssetReferencers;
+		AssetRegistry.GetDependencies(AssetData.PackageName, AssetReferencers);
+
+		for (const FAssetIdentifier& AssetReferencer : AssetReferencers)
+		{
+			UniqueDependencies.Add(AssetReferencer);
+		}
+	}
+
+	return UniqueDependencies.Array();
+}
+
+/*static*/ TMap<FString, TArray<FAssetIdentifier>> FPluginReferenceViewerUtils::SplitByPlugins(const TSharedRef<IPlugin>& InOwningPlugin, const TArray<FAssetIdentifier>& InPluginDependencies)
+{
+	TMap<FString, TArray<FAssetIdentifier>> Results;
+
+	const FName GameplayTagStructPackage = FGameplayTag::StaticStruct()->GetOutermost()->GetFName();
+	const FName NAME_GameplayTag = FGameplayTag::StaticStruct()->GetFName();
+
+	for (const FAssetIdentifier& AssetIdentifier : InPluginDependencies)
+	{
+		const FString PackageNameString = AssetIdentifier.PackageName.ToString();
+
+		if (AssetIdentifier.ObjectName == NAME_GameplayTag && AssetIdentifier.PackageName == GameplayTagStructPackage)
+		{
+			const TArray<TSharedRef<IPlugin>> SourcePlugins = FPluginReferenceViewerUtils::FindGameplayTagSourcePlugins(AssetIdentifier.ValueName);
+			for (const TSharedRef<IPlugin>& SourcePlugin : SourcePlugins)
+			{
+				if (SourcePlugin != InOwningPlugin)
+				{
+					Results.FindOrAdd(SourcePlugin->GetName()).Add(AssetIdentifier);
+				}
+			}
+		}
+		else if (FStringView ModuleName; FPackageName::TryConvertScriptPackageNameToModuleName(PackageNameString, ModuleName))
+		{
+			const TSharedPtr<IPlugin> ModulePlugin = IPluginManager::Get().GetModuleOwnerPlugin(FName(ModuleName));
+			if (ModulePlugin.IsValid() && ModulePlugin != InOwningPlugin)
+			{
+				Results.FindOrAdd(ModulePlugin->GetName()).Add(AssetIdentifier);
+			}
+		}
+		else
+		{
+			const TSharedPtr<IPlugin> PackagePlugin = IPluginManager::Get().FindPluginFromPath(PackageNameString);
+			if (PackagePlugin.IsValid() && PackagePlugin != InOwningPlugin)
+			{
+				Results.FindOrAdd(PackagePlugin->GetName()).Add(AssetIdentifier);
+			}
+		}
+	}
+
+	return Results;
+}
+
+/*static*/ void FPluginReferenceViewerUtils::SplitByReferenceType(const TArray<FAssetIdentifier>& InAssetIdentifiers, TArray<FAssetIdentifier>& OutAssetReferences, TArray<FAssetIdentifier>& OutScriptReferences, TArray<FAssetIdentifier>& OutNameReferences)
+{
+	const FName GameplayTagStructPackage = FGameplayTag::StaticStruct()->GetOutermost()->GetFName();
+	const FName NAME_GameplayTag = FGameplayTag::StaticStruct()->GetFName();
+	const FName NAME_DataTable(TEXT("DataTable"));
+
+	TArray<FName> AllPackageNames;
+	for (const FAssetIdentifier& AssetIdentifier : InAssetIdentifiers)
+	{
+		AllPackageNames.Add(AssetIdentifier.PackageName);
+	}
+
+	TMap<FName, FAssetData> PackagesToAssetDataMap;
+	UE::AssetRegistry::GetAssetForPackages(AllPackageNames, PackagesToAssetDataMap);
+
+	for (const FAssetIdentifier& AssetIdentifier : InAssetIdentifiers)
+	{
+		const FString PackageNameString = AssetIdentifier.PackageName.ToString();
+
+		if (AssetIdentifier.ObjectName == NAME_GameplayTag && AssetIdentifier.PackageName == GameplayTagStructPackage)
+		{
+			OutNameReferences.Add(AssetIdentifier);
+		}
+		else if (FPackageName::IsScriptPackage(PackageNameString))
+		{
+			FStringView ModuleName;
+			if (FPackageName::TryConvertScriptPackageNameToModuleName(PackageNameString, ModuleName))
+			{
+				OutScriptReferences.Add(AssetIdentifier);
+			}
+		}
+		else
+		{
+			// PackagesToAssetDataMap may not contain the package as references to assets that no longer exist do occur.
+			const FAssetData* AssetData = PackagesToAssetDataMap.Find(AssetIdentifier.PackageName);
+			if (AssetData != nullptr)
+			{
+				if (AssetData->AssetClassPath.GetAssetName() == NAME_DataTable)
+				{
+					OutNameReferences.Add(AssetIdentifier);
+				}
+				else
+				{
+					OutAssetReferences.Add(AssetIdentifier);
+				}
+			}
+			else
+			{
+				UE_LOG(LogPluginReferenceViewerUtils, Display, TEXT("Skipping package '%s' due to missing asset data. Package may no longer exist!"), *AssetIdentifier.PackageName.ToString());
+			}
+		}
+	}
+}
+
 /*static*/ void FPluginReferenceViewerUtils::ExportPlugins(const TArray<FString>& InPluginNames, const FString& InFilename)
 {
 	TArray<TSharedRef<IPlugin>> Plugins;
@@ -416,8 +416,8 @@ namespace PluginReferenceViewerCVars
 	{
 		SlowTask.EnterProgressFrame(1.0f, FText::Format(LOCTEXT("ExportPluginName", "Processing plugin {0}"), FText::FromString(CurrentPlugin->GetName())));
 
-		const TArray<FAssetIdentifier> AllDependencies = PluginReferenceViewerUtils::GetAssetDependencies(CurrentPlugin);
-		const TMap<FString, TArray<FAssetIdentifier>> PluginAssetMap = PluginReferenceViewerUtils::SplitByPlugins(CurrentPlugin, AllDependencies);
+		const TArray<FAssetIdentifier> AllDependencies = FPluginReferenceViewerUtils::GetAssetDependencies(CurrentPlugin);
+		const TMap<FString, TArray<FAssetIdentifier>> PluginAssetMap = FPluginReferenceViewerUtils::SplitByPlugins(CurrentPlugin, AllDependencies);
 
 		for (const FPluginReferenceDescriptor& ReferenceDescriptor : CurrentPlugin->GetDescriptor().Plugins)
 		{
@@ -429,7 +429,7 @@ namespace PluginReferenceViewerCVars
 			}
 
 			TArray<FAssetIdentifier> AssetReferences, ScriptReferences, NameReferences;
-			PluginReferenceViewerUtils::SplitByReferenceType(*PluginAssets, AssetReferences, ScriptReferences, NameReferences);
+			FPluginReferenceViewerUtils::SplitByReferenceType(*PluginAssets, AssetReferences, ScriptReferences, NameReferences);
 
 			StringBuilder.Reset();
 
@@ -517,15 +517,15 @@ namespace PluginReferenceViewerCVars
 
 	TStringBuilder<1024> StringBuilder;
 
-	const TArray<FAssetIdentifier> AllDependencies = PluginReferenceViewerUtils::GetAssetDependencies(Plugin.ToSharedRef());
-	const TMap<FString, TArray<FAssetIdentifier>> PluginAssetMap = PluginReferenceViewerUtils::SplitByPlugins(Plugin.ToSharedRef(), AllDependencies);
+	const TArray<FAssetIdentifier> AllDependencies = FPluginReferenceViewerUtils::GetAssetDependencies(Plugin.ToSharedRef());
+	const TMap<FString, TArray<FAssetIdentifier>> PluginAssetMap = FPluginReferenceViewerUtils::SplitByPlugins(Plugin.ToSharedRef(), AllDependencies);
 
 	// There might not be any asset references in the plugin dependency.
 	const TArray<FAssetIdentifier>* PluginAssets = PluginAssetMap.Find(InReference);
 	if (PluginAssets != nullptr)
 	{
 		TArray<FAssetIdentifier> AssetReferences, ScriptReferences, NameReferences;
-		PluginReferenceViewerUtils::SplitByReferenceType(*PluginAssets, AssetReferences, ScriptReferences, NameReferences);
+		FPluginReferenceViewerUtils::SplitByReferenceType(*PluginAssets, AssetReferences, ScriptReferences, NameReferences);
 
 		for (TArray<FAssetIdentifier>* AssetList : { &AssetReferences, &ScriptReferences, &NameReferences })
 		{
diff --git a/Engine/Plugins/Experimental/PluginReferenceViewer/Source/PluginReferenceViewer/Public/PluginReferenceViewerUtils.h b/Engine/Plugins/Experimental/PluginReferenceViewer/Source/PluginReferenceViewer/Public/PluginReferenceViewerUtils.h
--- a/Engine/Plugins/Experimental/PluginReferenceViewer/Source/PluginReferenceViewer/Public/PluginReferenceViewerUtils.h
+++ b/Engine/Plugins/Experimental/PluginReferenceViewer/Source/PluginReferenceViewer/Public/PluginReferenceViewerUtils.h
@@ -5,10 +5,20 @@
 #include "CoreMinimal.h"
 
 class IPlugin;
+struct FAssetIdentifier;
 
 class PLUGINREFERENCEVIEWER_API FPluginReferenceViewerUtils
 {
 public:
+	/* Helper that returns array of Asset dependencies for a given plugin*/
+	static TArray<FAssetIdentifier> GetAssetDependencies(const TSharedRef<IPlugin>& InPlugin);
+
+	/* Helper that will split plugin dependencies by their owning plugin*/
+	static TMap<FString, TArray<FAssetIdentifier>> SplitByPlugins(const TSharedRef<IPlugin>& InOwningPlugin, const TArray<FAssetIdentifier>& InPluginDependencies);
+
+	/* Helper that will split an array of assets by their reference types*/
+	static void SplitByReferenceType(const TArray<FAssetIdentifier>& InAssetIdentifiers, TArray<FAssetIdentifier>& OutAssetReferences, TArray<FAssetIdentifier>& OutScriptReferences, TArray<FAssetIdentifier>& OutNameReferences);
+
 	/* Exports the number of references from each plugin to their plugin dependencies. References include; assets, scripts and named references*/
 	static void ExportPlugins(const TArray<FString>& InPlugins, const FString& InFilename);
 
