Show friendly plugin names in the plugin reference viewer.

diff --git a/Engine/Plugins/Experimental/PluginReferenceViewer/Source/PluginReferenceViewer/Private/EdGraphNode_PluginReference.cpp b/Engine/Plugins/Experimental/PluginReferenceViewer/Source/PluginReferenceViewer/Private/EdGraphNode_PluginReference.cpp
--- a/Engine/Plugins/Experimental/PluginReferenceViewer/Source/PluginReferenceViewer/Private/EdGraphNode_PluginReference.cpp
+++ b/Engine/Plugins/Experimental/PluginReferenceViewer/Source/PluginReferenceViewer/Private/EdGraphNode_PluginReference.cpp
@@ -39,7 +39,7 @@ void UEdGraphNode_PluginReference::SetupPluginReferenceNode(const FIntPoint& InN
 
 	if (InPlugin.IsValid())
 	{
-		NodeTitle = FText::FromString(InPlugin->GetName());
+		NodeTitle = FText::FromString(InPlugin->GetFriendlyName());
 		bIsEnginePlugin = (InPlugin->GetType() == EPluginType::Engine) ? true : false;
 	}
 	else
diff --git a/Engine/Plugins/Experimental/PluginReferenceViewer/Source/PluginReferenceViewer/Private/SPluginReferenceViewer.cpp b/Engine/Plugins/Experimental/PluginReferenceViewer/Source/PluginReferenceViewer/Private/SPluginReferenceViewer.cpp
--- a/Engine/Plugins/Experimental/PluginReferenceViewer/Source/PluginReferenceViewer/Private/SPluginReferenceViewer.cpp
+++ b/Engine/Plugins/Experimental/PluginReferenceViewer/Source/PluginReferenceViewer/Private/SPluginReferenceViewer.cpp
@@ -401,12 +401,17 @@ FText SPluginReferenceViewer::GetPluginStatsText() const
 		return FText(LOCTEXT("PluginSelectionError", "Invalid number of nodes selected. Please select a single node!"));
 	}
 
-	FText PluginName = SelectedNode->GetNodeTitle(ENodeTitleType::FullTitle);
+	const UEdGraphNode_PluginReference* PluginReferenceNode = Cast<const UEdGraphNode_PluginReference>(SelectedNode);
 
-	const FPluginStats & PluginStats = GraphObj->GetPluginStats(FPluginIdentifier(FName(PluginName.ToString())));
+	if (PluginReferenceNode == nullptr)
+	{
+		return FText::GetEmpty();
+	}
+
+	const FPluginStats & PluginStats = GraphObj->GetPluginStats(PluginReferenceNode->GetIdentifier());
 
 	return FText::Format(LOCTEXT("PluginInfoText", "Name: {0}\nContent Size: {1}\nTotal Content Size With Dependencies: {2}\nDependencies: {3}\nReferencers: {4}"), 
-		{ PluginName, GetSizeInBytesAsBiggerUnits(PluginStats.Size), GetSizeInBytesAsBiggerUnits(PluginStats.SizeWithDependencies), PluginStats.Dependencies, PluginStats.Referencers });
+		{ FText::FromString(PluginReferenceNode->GetIdentifier().ToString()), GetSizeInBytesAsBiggerUnits(PluginStats.Size), GetSizeInBytesAsBiggerUnits(PluginStats.SizeWithDependencies), PluginStats.Dependencies, PluginStats.Referencers });
 }
 
 FText SPluginReferenceViewer::GetAddressBarText() const
