Add a new option to the plugin reference viewer utils to export all of the references by asset.
#rnx

diff --git a/Engine/Plugins/Experimental/PluginReferenceViewer/Source/PluginReferenceViewer/Private/PluginReferenceViewerUtils.cpp b/Engine/Plugins/Experimental/PluginReferenceViewer/Source/PluginReferenceViewer/Private/PluginReferenceViewerUtils.cpp
--- a/Engine/Plugins/Experimental/PluginReferenceViewer/Source/PluginReferenceViewer/Private/PluginReferenceViewerUtils.cpp
+++ b/Engine/Plugins/Experimental/PluginReferenceViewer/Source/PluginReferenceViewer/Private/PluginReferenceViewerUtils.cpp
@@ -32,7 +32,7 @@ namespace PluginReferenceViewerUtils
 		}
 		else
 		{
-			UE_LOG(LogPluginReferenceViewerUtils, Error, TEXT("Invalid plugin names argument. Expect plugin names seperated by ',' as the first argument. e.g a,b,c"));
+			UE_LOG(LogPluginReferenceViewerUtils, Error, TEXT("Invalid plugin names argument. Expect plugin names separated by ',' as the first argument. e.g a,b,c"));
 			return;
 		}
 
@@ -72,6 +72,30 @@ namespace PluginReferenceViewerUtils
 
 		FPluginReferenceViewerUtils::ExportReference(PluginName, ReferenceName, Filename);
 	}
+	
+	void ExportAssetReference(const TArray<FString>& InArgs)
+    	{
+    		if (InArgs.Num() < 2)
+    		{
+    			UE_LOG(LogPluginReferenceViewerUtils, Error, TEXT("Invalid arguments. Expected: [plugin name] [reference name] (optional)[filename.csv]"));
+    			return;
+    		}
+    
+    		FString PluginName = InArgs[0];
+    		FString ReferenceName = InArgs[1];
+    
+    		FString Filename;
+    		if (InArgs.Num() >= 3)
+    		{
+    			Filename = InArgs[2];
+    		}
+    		else
+    		{
+    			Filename = FPaths::ProjectSavedDir() / FPaths::SetExtension(FString::Printf(TEXT("%s-%s"), *PluginName, *ReferenceName), TEXT("csv"));
+    		}
+    
+    		FPluginReferenceViewerUtils::ExportAssetReference(PluginName, ReferenceName, Filename);
+    	}
 
 	void ExportDirectory(const TArray<FString>& InArgs)
 	{
@@ -203,11 +227,21 @@ namespace PluginReferenceViewerCVars
 	static FAutoConsoleCommand ExportPlugins(
 		TEXT("PluginReferenceViewer.ExportPlugins"),
 		TEXT("Exports to .csv the number of references (by type) that a plugin has for each of it's dependencies.\n"
-			"1st arg: single plugin name or multiple names seperated with ','.\n"
+			"1st arg: single plugin name or multiple names separated with ','.\n"
 			"2nd arg (optional): output filename.\n"
 			"Example: PluginReferenceViewer.ExportPlugins PluginA,PluginB,PluginC PluginReport.csv"),
 		FConsoleCommandWithArgsDelegate::CreateStatic(PluginReferenceViewerUtils::ExportPlugins)
 	);
+	
+	static FAutoConsoleCommand ExportAssetReference(
+    		TEXT("PluginReferenceViewer.ExportAssetReference"),
+    		TEXT("Exports to .csv the list of asset references that exist between a plugin and one of it's dependencies, including which asset has the reference.\n"
+    				"1st arg: plugin name.\n"
+    				"2nd arg: reference name.\n"
+    				"3rd arg (optional): output filename.\n"
+    				"Example: PluginReferenceViewer.ExportAssetReference PluginName ReferenceName PluginReport.csv"),
+    		FConsoleCommandWithArgsDelegate::CreateStatic(PluginReferenceViewerUtils::ExportAssetReference)
+    	);
 
 	static FAutoConsoleCommand ExportReference(
 		TEXT("PluginReferenceViewer.ExportReference"),
@@ -238,6 +272,36 @@ namespace PluginReferenceViewerCVars
 	);
 }
 
+/*static*/TMap<FAssetIdentifier, TArray<FAssetIdentifier>> FPluginReferenceViewerUtils::GetAssetDependencyMap(const TSharedRef<IPlugin>& InPlugin)
+{
+	IAssetRegistry& AssetRegistry = FModuleManager::LoadModuleChecked<FAssetRegistryModule>("AssetRegistry").Get();
+
+	TMap<FAssetIdentifier, TArray<FAssetIdentifier>> AssetDependencies;
+
+	FARFilter Filter;
+	Filter.bRecursivePaths = true;
+	Filter.PackagePaths.Add(FName(InPlugin->GetMountedAssetPath()));
+
+	TArray<FAssetData> AssetsInPlugin;
+	AssetRegistry.GetAssets(Filter, AssetsInPlugin);
+
+	TSet<FAssetIdentifier> UniqueDependencies;
+
+	for (const FAssetData& AssetData : AssetsInPlugin)
+	{
+		TArray<FAssetIdentifier> AssetReferencers;
+		AssetRegistry.GetDependencies(AssetData.PackageName, AssetReferencers);
+		AssetDependencies.Add(FAssetIdentifier(AssetData.PackageName), AssetReferencers);
+
+		for (const FAssetIdentifier& AssetReferencer : AssetReferencers)
+		{
+			UniqueDependencies.Add(AssetReferencer);
+		}
+	}
+
+	return AssetDependencies;
+}
+
 /*static*/TArray<FAssetIdentifier> FPluginReferenceViewerUtils::GetAssetDependencies(const TSharedRef<IPlugin>& InPlugin)
 {
 	IAssetRegistry& AssetRegistry = FModuleManager::LoadModuleChecked<FAssetRegistryModule>("AssetRegistry").Get();
@@ -478,6 +542,75 @@ namespace PluginReferenceViewerCVars
 	UE_LOG(LogPluginReferenceViewerUtils, Display, TEXT("Exported plugins; '%s' to '%s'"), *ConcatenatePluginNames, *FPaths::ConvertRelativePathToFull(InFilename));
 }
 
+/*static*/ void FPluginReferenceViewerUtils::ExportAssetReference(const FString& InPlugin, const FString& InReference, const FString& InFilename)
+{
+	const TSharedPtr<IPlugin> Plugin = IPluginManager::Get().FindPlugin(InPlugin);
+	if (!Plugin.IsValid())
+	{
+		UE_LOG(LogPluginReferenceViewerUtils, Error, TEXT("Plugin %s was not found!"), *InPlugin);
+		return;
+	}
+
+	const FPluginReferenceDescriptor* Found = Plugin->GetDescriptor().Plugins.FindByPredicate([=](const FPluginReferenceDescriptor& Item) { return Item.Name == InReference; });
+	if (Found == nullptr)
+	{
+		UE_LOG(LogPluginReferenceViewerUtils, Error, TEXT("Plugin reference %s was not found!"), *InReference);
+		return;
+	}
+
+	if (FPaths::GetExtension(InFilename) != TEXT("csv"))
+	{
+		UE_LOG(LogPluginReferenceViewerUtils, Error, TEXT("Invalid filename extenstion '%s'. Expected .csv"), *InFilename);
+		return;
+	}
+
+	IAssetRegistry* AssetRegistry = IAssetRegistry::Get();
+	AssetRegistry->WaitForCompletion();
+
+	TUniquePtr<IFileHandle> ExportFileHandle(FPlatformFileManager::Get().GetPlatformFile().OpenWrite(*InFilename));
+	if (!ExportFileHandle.IsValid())
+	{
+		return;
+	}
+
+	UTF16CHAR BOM = UNICODE_BOM;
+	ExportFileHandle->Write((uint8*)&BOM, sizeof(UTF16CHAR));
+
+	constexpr TCHAR Separator = TEXT(',');
+	constexpr TCHAR LineEnd = TEXT('\n');
+
+	TStringBuilder<1024> StringBuilder;
+
+	const TMap<FAssetIdentifier, TArray<FAssetIdentifier>> AllDependencies = FPluginReferenceViewerUtils::GetAssetDependencyMap(Plugin.ToSharedRef());
+	for (const TPair<FAssetIdentifier, TArray<FAssetIdentifier>>& AssetDependency : AllDependencies)
+	{
+		const TMap<FString, TArray<FAssetIdentifier>> PluginAssetMap = FPluginReferenceViewerUtils::SplitByPlugins(Plugin.ToSharedRef(), AssetDependency.Value);
+    
+    	// There might not be any asset references in the plugin dependency.
+    	const TArray<FAssetIdentifier>* PluginAssets = PluginAssetMap.Find(InReference);
+    	if (PluginAssets != nullptr)
+    	{
+    		TArray<FAssetIdentifier> AssetReferences, ScriptReferences, NameReferences;
+    		FPluginReferenceViewerUtils::SplitByReferenceType(*PluginAssets, AssetReferences, ScriptReferences, NameReferences);
+    
+    		for (TArray<FAssetIdentifier>* AssetList : { &AssetReferences, &ScriptReferences, &NameReferences })
+    		{
+    			for (const FAssetIdentifier& Identifier : *AssetList)
+    			{
+    				StringBuilder.Reset();
+    				StringBuilder.Append(AssetDependency.Key.ToString());
+    				StringBuilder.AppendChar(Separator);
+    				StringBuilder.Append(Identifier.ToString());
+    				StringBuilder.AppendChar(LineEnd);
+    				ExportFileHandle->Write((const uint8*)StringBuilder.ToString(), StringBuilder.Len() * sizeof(TCHAR));
+    			}
+    		}
+    	}
+	}
+
+	ExportFileHandle->Flush();
+}
+
 /*static*/ void FPluginReferenceViewerUtils::ExportReference(const FString& InPlugin, const FString& InReference, const FString& InFilename)
 {
 	const TSharedPtr<IPlugin> Plugin = IPluginManager::Get().FindPlugin(InPlugin);
diff --git a/Engine/Plugins/Experimental/PluginReferenceViewer/Source/PluginReferenceViewer/Public/PluginReferenceViewerUtils.h b/Engine/Plugins/Experimental/PluginReferenceViewer/Source/PluginReferenceViewer/Public/PluginReferenceViewerUtils.h
--- a/Engine/Plugins/Experimental/PluginReferenceViewer/Source/PluginReferenceViewer/Public/PluginReferenceViewerUtils.h
+++ b/Engine/Plugins/Experimental/PluginReferenceViewer/Source/PluginReferenceViewer/Public/PluginReferenceViewerUtils.h
@@ -12,6 +12,9 @@ struct FAssetIdentifier;
 class FPluginReferenceViewerUtils
 {
 public:
+	/* Helper that returns array of Asset dependencies for a given plugin*/
+	static UE_API TMap<FAssetIdentifier, TArray<FAssetIdentifier>> GetAssetDependencyMap(const TSharedRef<IPlugin>& InPlugin);
+
 	/* Helper that returns array of Asset dependencies for a given plugin*/
 	static UE_API TArray<FAssetIdentifier> GetAssetDependencies(const TSharedRef<IPlugin>& InPlugin);
 
@@ -24,6 +27,9 @@ public:
 	/* Exports the number of references from each plugin to their plugin dependencies. References include; assets, scripts and named references*/
 	static UE_API void ExportPlugins(const TArray<FString>& InPlugins, const FString& InFilename);
 
+	/* Exports to a .csv file the list of asset references that exist between the plugin and one of it's plugin dependencies by asset */
+	static UE_API void ExportAssetReference(const FString& InPlugin, const FString& InReference, const FString& InFilename);
+
 	/* Exports to a .csv file the list of asset references that exist between the plugin and one of it's plugin dependencies */
 	static UE_API void ExportReference(const FString& InPlugin, const FString& InReference, const FString& InFilename);
 
